Initial check-in of editor menus module for blueprint/script and customization

diff --git a/Engine/Source/Editor/EditorMenus/EditorMenus.Build.cs b/Engine/Source/Editor/EditorMenus/EditorMenus.Build.cs
new file mode 100644
--- /dev/null
+++ b/Engine/Source/Editor/EditorMenus/EditorMenus.Build.cs
@@ -0,0 +1,30 @@
+// Copyright 1998-2019 Epic Games, Inc. All Rights Reserved.
+
+namespace UnrealBuildTool.Rules
+{
+	public class EditorMenus : ModuleRules
+	{
+		public EditorMenus(ReadOnlyTargetRules Target) : base(Target)
+		{
+			//bFasterWithoutUnity = true;
+			//PCHUsage = PCHUsageMode.NoPCHs;
+
+			PrivateDependencyModuleNames.AddRange(
+				new string[] {
+					"Core",
+					"CoreUObject",
+					"Slate",
+					"SlateCore",
+					"EditorSubSystem",
+					"Engine", // For Subsystems
+					"UnrealEd",
+				}
+			);
+
+			//PrivateIncludePathModuleNames.AddRange(
+			//	new string[] {
+			//	"UnrealEd",
+			//	});
+		}
+	}
+}
diff --git a/Engine/Source/Editor/EditorMenus/Private/EditorMenu.cpp b/Engine/Source/Editor/EditorMenus/Private/EditorMenu.cpp
new file mode 100644
--- /dev/null
+++ b/Engine/Source/Editor/EditorMenus/Private/EditorMenu.cpp
@@ -0,0 +1,218 @@
+// Copyright 1998-2019 Epic Games, Inc. All Rights Reserved.
+
+#include "EditorMenu.h"
+#include "EditorMenuSubsystem.h"
+#include "IEditorMenusModule.h"
+
+#include "Textures/SlateIcon.h"
+#include "Framework/Commands/UIAction.h"
+#include "Framework/MultiBox/MultiBox.h"
+#include "Framework/MultiBox/MultiBoxBuilder.h"
+#include "Internationalization/Internationalization.h"
+#include "AssetTypeCategories.h"
+
+#include "Editor.h"
+
+UEditorMenu::UEditorMenu() :
+	MenuType(EMultiBoxType::Menu)
+	, bShouldCloseWindowAfterMenuSelection(true)
+	, bCloseSelfOnly(false)
+	, bSearchable(false)
+	, bToolBarIsFocusable(false)
+	, bToolBarForceSmallIcons(false)
+	, bRegistered(false)
+	, StyleSet(&FCoreStyle::Get())
+{
+}
+
+void UEditorMenu::InitMenu(const FEditorMenuOwner InOwner, FName InName, FName InParent, EMultiBoxType InType)
+{
+	MenuOwner = InOwner;
+	MenuName = InName;
+	MenuParent = InParent;
+	MenuType = InType;
+}
+
+void UEditorMenu::InitGeneratedCopy(const UEditorMenu* Source)
+{
+	// Skip sections and context
+
+	MenuName = Source->MenuName;
+	MenuParent = Source->MenuParent;
+	StyleName = Source->StyleName;
+	TutorialHighlightName = Source->TutorialHighlightName;
+	MenuType = Source->MenuType;
+	StyleSet = Source->StyleSet;
+	bShouldCloseWindowAfterMenuSelection = Source->bShouldCloseWindowAfterMenuSelection;
+	bCloseSelfOnly = Source->bCloseSelfOnly;
+	bSearchable = Source->bSearchable;
+	bToolBarIsFocusable = Source->bToolBarIsFocusable;
+	bToolBarForceSmallIcons = Source->bToolBarForceSmallIcons;
+	MenuOwner = Source->MenuOwner;
+}
+
+int32 UEditorMenu::IndexOfSection(const FName InSectionName) const
+{
+	for (int32 i=0; i < Sections.Num(); ++i)
+	{
+		if (Sections[i].Name == InSectionName)
+		{
+			return i;
+		}
+	}
+
+	return INDEX_NONE;
+}
+
+int32 UEditorMenu::FindInsertIndex(const FEditorMenuSection& InSection) const
+{
+	const FEditorMenuInsert InInsertPosition = InSection.InsertPosition;
+	if (InInsertPosition.IsDefault())
+	{
+		return Sections.Num();
+	}
+
+	if (InInsertPosition.Position == EEditorMenuInsertType::First)
+	{
+		for (int32 i = 0; i < Sections.Num(); ++i)
+		{
+			if (Sections[i].InsertPosition.Position != InInsertPosition.Position)
+			{
+				return i;
+			}
+		}
+
+		return Sections.Num();
+	}
+
+	int32 DestIndex = IndexOfSection(InInsertPosition.Name);
+	if (DestIndex == INDEX_NONE)
+	{
+		return DestIndex;
+	}
+
+	if (InInsertPosition.Position == EEditorMenuInsertType::After)
+	{
+		++DestIndex;
+	}
+
+	for (int32 i = DestIndex; i < Sections.Num(); ++i)
+	{
+		if (Sections[i].InsertPosition != InInsertPosition)
+		{
+			return i;
+		}
+	}
+
+	return Sections.Num();
+}
+
+FEditorMenuSection& UEditorMenu::AddDynamicSection(const FName SectionName, const FNewSectionConstructChoice& InConstruct, const FEditorMenuInsert InPosition)
+{
+	FEditorMenuSection& Section = AddSection(SectionName, TAttribute< FText >(), InPosition);
+	Section.Construct = InConstruct;
+	return Section;
+}
+
+FEditorMenuSection& UEditorMenu::AddSection(const FName SectionName, const TAttribute< FText >& InLabel, const FEditorMenuInsert InPosition)
+{
+	for (FEditorMenuSection& Section : Sections)
+	{
+		if (Section.Name == SectionName)
+		{
+			if (InLabel.IsSet())
+			{
+				Section.Label = InLabel;
+			}
+
+			if (InPosition.Name != NAME_None)
+			{
+				Section.InsertPosition = InPosition;
+			}
+
+			return Section;
+		}
+	}
+
+	FEditorMenuSection& NewSection = Sections.AddDefaulted_GetRef();
+	NewSection.InitSection(SectionName, InLabel, InPosition);
+	return NewSection;
+}
+
+void UEditorMenu::AddSectionScript(const FName SectionName, const FText& InLabel, const FName InsertName, const EEditorMenuInsertType InsertType)
+{
+	FEditorMenuSection& Section = FindOrAddSection(SectionName);
+	Section.Label = InLabel;
+	Section.InsertPosition = FEditorMenuInsert(InsertName, InsertType);
+}
+
+void UEditorMenu::AddDynamicSectionScript(const FName SectionName, UEditorMenuSectionDynamic* InObject)
+{
+	FEditorMenuSection& Section = FindOrAddSection(SectionName);
+	Section.EditorMenuSectionDynamic = InObject;
+}
+
+void UEditorMenu::AddMenuEntryObject(UEditorMenuEntryScript* InObject)
+{
+	FindOrAddSection(InObject->Data.Section).AddEntryObject(InObject);
+}
+
+UEditorMenu* UEditorMenu::AddSubMenu(const FEditorMenuOwner InOwner, const FName SectionName, const FName InName, const FText& InLabel, const FText& InToolTip)
+{
+	FEditorMenuEntry Args = FEditorMenuEntry::InitMenuEntry(InName, InLabel, InToolTip, FSlateIcon(), FUIAction());
+	Args.Owner = InOwner;
+	Args.SubMenu.bIsSubMenu = true;
+	FindOrAddSection(NAME_None).AddEntry(Args);
+	return UEditorMenuSubsystem::Get()->ExtendMenu(*(MenuName.ToString() + TEXT(".") + InName.ToString()));
+}
+
+FEditorMenuSection* UEditorMenu::FindSection(const FName SectionName)
+{
+	for (FEditorMenuSection& Section : Sections)
+	{
+		if (Section.Name == SectionName)
+		{
+			return &Section;
+		}
+	}
+
+	return nullptr;
+}
+
+FEditorMenuSection& UEditorMenu::FindOrAddSection(const FName SectionName)
+{
+	for (FEditorMenuSection& Section : Sections)
+	{
+		if (Section.Name == SectionName)
+		{
+			return Section;
+		}
+	}
+	
+	return AddSection(SectionName);
+}
+
+void UEditorMenu::RemoveSection(const FName SectionName)
+{
+	Sections.RemoveAll([SectionName](const FEditorMenuSection& Section) { return Section.Name == SectionName; });
+}
+
+bool UEditorMenu::FindEntry(const FName EntryName, int32& SectionIndex, int32& EntryIndex) const
+{
+	for (int32 i=0; i < Sections.Num(); ++i)
+	{
+		EntryIndex = Sections[i].IndexOfBlock(EntryName);
+		if (EntryIndex != INDEX_NONE)
+		{
+			SectionIndex = i;
+			return true;
+		}
+	}
+
+	return false;
+}
+
+void UEditorMenu::AddMenuEntry(const FName SectionName, const FEditorMenuEntry& Args)
+{
+	FindOrAddSection(SectionName).AddEntry(Args);
+}
diff --git a/Engine/Source/Editor/EditorMenus/Private/EditorMenuContext.cpp b/Engine/Source/Editor/EditorMenus/Private/EditorMenuContext.cpp
new file mode 100644
--- /dev/null
+++ b/Engine/Source/Editor/EditorMenus/Private/EditorMenuContext.cpp
@@ -0,0 +1,101 @@
+// Copyright 1998-2019 Epic Games, Inc. All Rights Reserved.
+
+#include "EditorMenuContext.h"
+#include "EditorMenuSubsystem.h"
+#include "IEditorMenusModule.h"
+
+#include "Textures/SlateIcon.h"
+#include "Framework/Commands/UIAction.h"
+#include "Framework/MultiBox/MultiBox.h"
+#include "Framework/MultiBox/MultiBoxBuilder.h"
+#include "Internationalization/Internationalization.h"
+#include "Toolkits/AssetEditorToolkit.h"
+
+#include "Editor.h"
+
+FEditorMenuContext::FEditorMenuContext()
+{
+
+}
+
+FEditorMenuContext::FEditorMenuContext(UObject* InContext)
+{
+	if (InContext)
+	{
+		ContextObjects.Add(InContext);
+	}
+}
+
+FEditorMenuContext::FEditorMenuContext(TSharedPtr<FUICommandList> InCommandList, TSharedPtr<FExtender> InExtender, UObject* InContext)
+{
+	if (InContext)
+	{
+		ContextObjects.Add(InContext);
+	}
+
+	if (InExtender.IsValid())
+	{
+		AddExtender(InExtender);
+	}
+
+	AppendCommandList(InCommandList);
+}
+
+UObject* FEditorMenuContext::FindByClass(UClass* InClass) const
+{
+	for (UObject* ContextObject : ContextObjects)
+	{
+		if (ContextObject && ContextObject->IsA(InClass))
+		{
+			return ContextObject;
+		}
+	}
+
+	return nullptr;
+}
+
+void FEditorMenuContext::AppendCommandList(TSharedPtr<FUICommandList> InCommandList)
+{
+	if (InCommandList.IsValid())
+	{
+		if (!CommandList.IsValid())
+		{
+			CommandList = MakeShared<FUICommandList>();
+		}
+
+		CommandList->Append(InCommandList.ToSharedRef());
+	}
+}
+
+void FEditorMenuContext::AddExtender(const TSharedPtr<FExtender>& InExtender)
+{
+	if (!ExtensibilityManager.IsValid())
+	{
+		ExtensibilityManager = MakeShared<FExtensibilityManager>();
+	}
+
+	ExtensibilityManager->AddExtender(InExtender);
+}
+
+TSharedPtr<FExtender> FEditorMenuContext::GetAllExtenders()
+{
+	if (ExtensibilityManager.IsValid())
+	{
+		return ExtensibilityManager->GetAllExtenders();
+	}
+
+	return TSharedPtr<FExtender>();
+}
+
+void FEditorMenuContext::AppendObjects(const TArray<UObject*>& InObjects)
+{
+	for (UObject* Object : InObjects)
+	{
+		AddObject(Object);
+	}
+}
+
+void FEditorMenuContext::AddObject(UObject* InObject)
+{
+	ContextObjects.AddUnique(InObject);
+}
diff --git a/Engine/Source/Editor/EditorMenus/Private/EditorMenuEntry.cpp b/Engine/Source/Editor/EditorMenus/Private/EditorMenuEntry.cpp
new file mode 100644
--- /dev/null
+++ b/Engine/Source/Editor/EditorMenus/Private/EditorMenuEntry.cpp
@@ -0,0 +1,128 @@
+// Copyright 1998-2019 Epic Games, Inc. All Rights Reserved.
+
+#include "EditorMenuEntry.h"
+#include "EditorMenuSubsystem.h"
+#include "IEditorMenusModule.h"
+
+#include "Textures/SlateIcon.h"
+#include "Framework/Commands/UIAction.h"
+#include "Framework/MultiBox/MultiBox.h"
+#include "Framework/MultiBox/MultiBoxBuilder.h"
+#include "Internationalization/Internationalization.h"
+
+#include "Editor.h"
+
+
+FEditorMenuEntry::FEditorMenuEntry() :
+	Type(EMultiBlockType::None),
+	UserInterfaceActionType(EUserInterfaceActionType::Button),
+	bShouldCloseWindowAfterMenuSelection(true),
+	ScriptObject(nullptr)
+{
+}
+
+FEditorMenuEntry::FEditorMenuEntry(const FEditorMenuOwner InOwner, const FName InName, EMultiBlockType InType) :
+	Name(InName),
+	Owner(InOwner),
+	Type(InType),
+	UserInterfaceActionType(EUserInterfaceActionType::Button),
+	bShouldCloseWindowAfterMenuSelection(true)	,
+	ScriptObject(nullptr)
+{
+}
+
+void FEditorMenuEntry::SetCommand(const TSharedPtr< const FUICommandInfo >& InCommand, FName InName, const TAttribute<FText>& InLabel, const TAttribute<FText>& InToolTip, const TAttribute<FSlateIcon>& InIcon)
+{
+	//Action = FEditorUIActionChoice(InCommand, InCommandList);
+
+	Command = InCommand;
+	Name = InName != NAME_None ? InName : InCommand->GetCommandName();
+	Label = InLabel.IsSet() ? InLabel : InCommand->GetLabel();
+	ToolTip = InToolTip.IsSet() ? InToolTip : InCommand->GetDescription();
+	Icon = InIcon.IsSet() ? InIcon : InCommand->GetIcon();
+}
+
+FEditorMenuEntry FEditorMenuEntry::InitMenuEntry(const FName InName, const TAttribute<FText>& InLabel, const TAttribute<FText>& InToolTip, const TAttribute<FSlateIcon>& InIcon, const FEditorUIActionChoice& InAction, const EUserInterfaceActionType InUserInterfaceActionType, const FName InTutorialHighlightName)
+{
+	FEditorMenuEntry Entry(UEditorMenuSubsystem::Get()->CurrentOwner(), InName, EMultiBlockType::MenuEntry);
+	Entry.TutorialHighlightName = InTutorialHighlightName;
+	Entry.Label = InLabel;
+	Entry.ToolTip = InToolTip;
+	Entry.Icon = InIcon;
+	Entry.UserInterfaceActionType = InUserInterfaceActionType;
+	Entry.Action = InAction;
+	return Entry;
+}
+
+FEditorMenuEntry FEditorMenuEntry::InitMenuEntry(const TSharedPtr< const FUICommandInfo >& InCommand, const TAttribute<FText>& InLabel, const TAttribute<FText>& InToolTip, const TAttribute<FSlateIcon>& InIcon, const FName InTutorialHighlightName, const FName InName)
+{
+	FEditorMenuEntry Entry(UEditorMenuSubsystem::Get()->CurrentOwner(), InName, EMultiBlockType::MenuEntry);
+	Entry.TutorialHighlightName = InTutorialHighlightName;
+	Entry.SetCommand(InCommand, InName, InLabel, InToolTip, InIcon);
+	return Entry;
+}
+
+FEditorMenuEntry FEditorMenuEntry::InitSubMenu(const FName InParentMenu, const FName InName, const TAttribute<FText>& InLabel, const TAttribute<FText>& InToolTip, const FNewEditorMenuChoice& InMakeMenu, bool bInOpenSubMenuOnClick, const TAttribute<FSlateIcon>& InIcon, const bool bInShouldCloseWindowAfterMenuSelection)
+{
+	FEditorMenuEntry Entry(UEditorMenuSubsystem::Get()->CurrentOwner(), InName, EMultiBlockType::MenuEntry);
+	Entry.Label = InLabel;
+	Entry.ToolTip = InToolTip;
+	Entry.Icon = InIcon;
+
+	Entry.SubMenu = FSubMenuEntryData(InMakeMenu, bInOpenSubMenuOnClick);
+	Entry.bShouldCloseWindowAfterMenuSelection = bInShouldCloseWindowAfterMenuSelection;
+	return Entry;
+}
+
+FEditorMenuEntry FEditorMenuEntry::InitToolBarButton(const FName InName, const FEditorUIActionChoice& InAction, const TAttribute<FText>& InLabel, const TAttribute<FText>& InToolTip, const TAttribute<FSlateIcon>& InIcon, const EUserInterfaceActionType InUserInterfaceActionType, FName InTutorialHighlightName)
+{
+	FEditorMenuEntry Entry(UEditorMenuSubsystem::Get()->CurrentOwner(), InName, EMultiBlockType::ToolBarButton);
+	Entry.TutorialHighlightName = InTutorialHighlightName;
+	Entry.Label = InLabel;
+	Entry.ToolTip = InToolTip;
+	Entry.Icon = InIcon;
+	Entry.UserInterfaceActionType = InUserInterfaceActionType;
+	Entry.Action = InAction;
+	return Entry;
+}
+
+FEditorMenuEntry FEditorMenuEntry::InitToolBarButton(const TSharedPtr< const FUICommandInfo >& InCommand, const TAttribute<FText>& InLabel, const TAttribute<FText>& InToolTip, const TAttribute<FSlateIcon>& InIcon, FName InTutorialHighlightName, const FName InName)
+{
+	FEditorMenuEntry Entry(UEditorMenuSubsystem::Get()->CurrentOwner(), InName, EMultiBlockType::ToolBarButton);
+	Entry.TutorialHighlightName = InTutorialHighlightName;
+	Entry.SetCommand(InCommand, InName, InLabel, InToolTip, InIcon);
+	return Entry;
+}
+
+FEditorMenuEntry FEditorMenuEntry::InitComboButton(const FName InName, const FEditorUIActionChoice& InAction, const FOnGetContent& InMenuContentGenerator, const TAttribute<FText>& InLabel, const TAttribute<FText>& InToolTip, const TAttribute<FSlateIcon>& InIcon, bool bInSimpleComboBox, FName InTutorialHighlightName)
+{
+	FEditorMenuEntry Entry(UEditorMenuSubsystem::Get()->CurrentOwner(), InName, EMultiBlockType::ToolBarComboButton);
+	Entry.TutorialHighlightName = InTutorialHighlightName;
+	Entry.Label = InLabel;
+	Entry.ToolTip = InToolTip;
+	Entry.Icon = InIcon;
+	Entry.Action = InAction;
+
+	Entry.ToolBar.ComboButtonContextMenuGenerator = InMenuContentGenerator;
+	Entry.ToolBar.bSimpleComboBox = bInSimpleComboBox;
+	return Entry;
+}
+
+FEditorMenuEntry FEditorMenuEntry::InitMenuSeparator(const FName InName)
+{
+	return FEditorMenuEntry(UEditorMenuSubsystem::Get()->CurrentOwner(), InName, EMultiBlockType::MenuSeparator);
+}
+
+FEditorMenuEntry FEditorMenuEntry::InitToolBarSeparator(const FName InName)
+{
+	return FEditorMenuEntry(UEditorMenuSubsystem::Get()->CurrentOwner(), InName, EMultiBlockType::ToolBarSeparator);
+}
+
+void FEditorMenuEntry::ResetActions()
+{
+	Action = FEditorUIActionChoice();
+	//ScriptObject = nullptr;
+	StringCommand = FEditorMenuStringCommand();
+	Command.Reset();
+	CommandList.Reset();
+}
diff --git a/Engine/Source/Editor/EditorMenus/Private/EditorMenuEntryScript.cpp b/Engine/Source/Editor/EditorMenus/Private/EditorMenuEntryScript.cpp
new file mode 100644
--- /dev/null
+++ b/Engine/Source/Editor/EditorMenus/Private/EditorMenuEntryScript.cpp
@@ -0,0 +1,146 @@
+// Copyright 1998-2019 Epic Games, Inc. All Rights Reserved.
+
+#include "EditorMenuEntryScript.h"
+#include "EditorMenuEntry.h"
+#include "EditorMenuSubsystem.h"
+#include "IEditorMenusModule.h"
+
+#include "Textures/SlateIcon.h"
+#include "Framework/Commands/UIAction.h"
+#include "Framework/MultiBox/MultiBox.h"
+#include "Framework/MultiBox/MultiBoxBuilder.h"
+#include "Internationalization/Internationalization.h"
+
+#include "Editor.h"
+
+FScriptSlateIcon::FScriptSlateIcon()
+{
+
+}
+
+FScriptSlateIcon::FScriptSlateIcon(const FName InStyleSetName, const FName InStyleName) :
+	StyleSetName(InStyleSetName),
+	StyleName(InStyleName),
+	SmallStyleName(ISlateStyle::Join(InStyleName, ".Small"))
+{
+}
+
+FScriptSlateIcon::FScriptSlateIcon(const FName InStyleSetName, const FName InStyleName, const FName InSmallStyleName) :
+	StyleSetName(InStyleSetName),
+	StyleName(InStyleName),
+	SmallStyleName(InSmallStyleName)
+{
+}
+
+FSlateIcon FScriptSlateIcon::GetSlateIcon() const
+{
+	if (SmallStyleName == NAME_None)
+	{
+		if (StyleSetName == NAME_None && StyleName == NAME_None)
+		{
+			return FSlateIcon();
+		}
+
+		return FSlateIcon(StyleSetName, StyleName);
+	}
+
+	return FSlateIcon(StyleSetName, StyleName, SmallStyleName);
+}
+
+TAttribute<FText> UEditorMenuEntryScript::CreateLabelAttribute(FEditorMenuContext& Context)
+{
+	static const FName FunctionName = GET_FUNCTION_NAME_CHECKED(UEditorMenuEntryScript, GetLabel);
+	if (GetClass()->IsFunctionImplementedInScript(FunctionName))
+	{
+		return TAttribute<FText>::Create(TAttribute<FText>::FGetter::CreateUFunction(this, FunctionName, Context));
+	}
+
+	return Data.Label;
+}
+
+TAttribute<FText> UEditorMenuEntryScript::CreateToolTipAttribute(FEditorMenuContext& Context)
+{
+	static const FName FunctionName = GET_FUNCTION_NAME_CHECKED(UEditorMenuEntryScript, GetToolTip);
+	if (GetClass()->IsFunctionImplementedInScript(FunctionName))
+	{
+		return TAttribute<FText>::Create(TAttribute<FText>::FGetter::CreateUFunction(this, FunctionName, Context));
+	}
+
+	return Data.ToolTip;
+}
+
+void UEditorMenuEntryScript::RegisterMenuEntry()
+{
+	UEditorMenuSubsystem::AddMenuEntryObject(this);
+}
+
+void UEditorMenuEntryScript::InitEntry(const FName OwnerName, const FName Menu, const FName Section, const FName Name, const FText& Label, const FText& ToolTip)
+{
+	Data.OwnerName = OwnerName;
+	Data.Menu = Menu;
+	Data.Section = Section;
+	Data.Name = Name;
+	Data.Label = Label;
+	Data.ToolTip = ToolTip;
+}
+
+void UEditorMenuEntryScript::ToMenuEntry(FEditorMenuEntry& Output)
+{
+	if (Data.Advanced.bIsSubMenu)
+	{
+		Output = FEditorMenuEntry::InitSubMenu(
+			Data.Menu,
+			Data.Name,
+			Data.Label,
+			Data.ToolTip,
+			FNewEditorMenuChoice(), // Menu will be opened by string: 'Menu' + '.' + 'Name'
+			Data.Advanced.bOpenSubMenuOnClick,
+			Data.Icon,
+			Data.Advanced.bShouldCloseWindowAfterMenuSelection);
+	}
+	else
+	{
+		if (Data.Advanced.EntryType == EMultiBlockType::ToolBarButton)
+		{
+			Output = FEditorMenuEntry::InitToolBarButton(
+				Data.Name,
+				FEditorUIActionChoice(), // Action will be handled by 'ScriptObject'
+				Data.Label,
+				Data.ToolTip,
+				Data.Icon,
+				Data.Advanced.UserInterfaceActionType,
+				Data.Advanced.TutorialHighlight
+			);
+		}
+		else
+		{
+			Output = FEditorMenuEntry::InitMenuEntry(Data.Name, Data.Label, Data.ToolTip, Data.Icon, FUIAction());
+		}
+	}
+
+	if (!Data.InsertPosition.IsDefault())
+	{
+		Output.InsertPosition = Data.InsertPosition;
+	}
+	
+	Output.ScriptObject = this;
+
+	Output.Owner = Data.OwnerName;
+}
+
+bool FEditorMenuEntry::IsScriptObjectDynamicConstruct() const
+{
+	static const FName ConstructMenuEntryName = GET_FUNCTION_NAME_CHECKED(UEditorMenuEntryScript, ConstructMenuEntry);
+	return ScriptObject && ScriptObject->GetClass()->IsFunctionImplementedInScript(ConstructMenuEntryName);
+}
+
+FEditorMenuEntryScriptDataAdvanced::FEditorMenuEntryScriptDataAdvanced() :
+	EntryType(EMultiBlockType::MenuEntry),
+	UserInterfaceActionType(EUserInterfaceActionType::Button),
+	bIsSubMenu(false),
+	bOpenSubMenuOnClick(false),
+	bShouldCloseWindowAfterMenuSelection(true),
+	bSimpleComboBox(false)
+{
+
+}
diff --git a/Engine/Source/Editor/EditorMenus/Private/EditorMenuSection.cpp b/Engine/Source/Editor/EditorMenus/Private/EditorMenuSection.cpp
new file mode 100644
--- /dev/null
+++ b/Engine/Source/Editor/EditorMenus/Private/EditorMenuSection.cpp
@@ -0,0 +1,197 @@
+// Copyright 1998-2019 Epic Games, Inc. All Rights Reserved.
+
+#include "EditorMenuSection.h"
+#include "EditorMenuSubsystem.h"
+#include "IEditorMenusModule.h"
+
+#include "Textures/SlateIcon.h"
+#include "Framework/Commands/UIAction.h"
+#include "Framework/MultiBox/MultiBox.h"
+#include "Framework/MultiBox/MultiBoxBuilder.h"
+#include "Internationalization/Internationalization.h"
+
+#include "Editor.h"
+
+FEditorMenuSection::FEditorMenuSection() :
+	EditorMenuSectionDynamic(nullptr)
+{
+
+}
+
+
+void FEditorMenuSection::InitSection(const FName InName, const TAttribute< FText >& InLabel, const FEditorMenuInsert InPosition)
+{
+	Name = InName;
+	Label = InLabel;
+	InsertPosition = InPosition;
+}
+
+void FEditorMenuSection::InitGeneratedSectionCopy(const FEditorMenuSection& Source, FEditorMenuContext& InContext)
+{
+	Name = Source.Name;
+	Label = Source.Label;
+	InsertPosition = Source.InsertPosition;
+	Construct = Source.Construct;
+	Context = InContext;
+}
+
+FEditorMenuEntry& FEditorMenuSection::AddEntry(const FEditorMenuEntry& Args)
+{
+	if (Args.Name == NAME_None)
+	{
+		return Blocks.Add_GetRef(Args);
+	}
+
+	int32 BlockIndex = IndexOfBlock(Args.Name);
+	if (BlockIndex != INDEX_NONE)
+	{
+		Blocks[BlockIndex] = Args;
+		return Blocks[BlockIndex];
+	}
+	else
+	{
+		return Blocks.Add_GetRef(Args);
+	}
+}
+
+FEditorMenuEntry& FEditorMenuSection::AddEntryObject(UEditorMenuEntryScript* InObject)
+{
+	// Avoid modifying objects that are saved as content on disk
+	UEditorMenuEntryScript* DestObject = InObject;
+	if (DestObject->IsAsset())
+	{
+		DestObject = DuplicateObject<UEditorMenuEntryScript>(InObject, UEditorMenuSubsystem::Get());
+	}
+
+	FEditorMenuEntry Args;
+	DestObject->ToMenuEntry(Args);
+	return AddEntry(Args);
+}
+
+FEditorMenuEntry& FEditorMenuSection::AddMenuEntry(const TSharedPtr< const FUICommandInfo >& InCommand)
+{
+	return AddEntry(FEditorMenuEntry::InitMenuEntry(InCommand));
+}
+
+FEditorMenuEntry& FEditorMenuSection::AddDynamicEntry(const FName InName, const FNewEditorMenuSectionDelegate& InConstruct)
+{
+	FEditorMenuEntry& Entry = AddEntry(FEditorMenuEntry(UEditorMenuSubsystem::Get()->CurrentOwner(), InName, EMultiBlockType::MenuEntry));
+	Entry.Construct = InConstruct;
+	return Entry;
+}
+
+FEditorMenuEntry& FEditorMenuSection::AddDynamicEntry(const FName InName, const FNewEditorMenuDelegateLegacy& InConstruct)
+{
+	FEditorMenuEntry& Entry = AddEntry(FEditorMenuEntry(UEditorMenuSubsystem::Get()->CurrentOwner(), InName, EMultiBlockType::MenuEntry));
+	Entry.ConstructLegacy = InConstruct;
+	return Entry;
+}
+
+FEditorMenuEntry& FEditorMenuSection::AddMenuSeparator(const FName InName)
+{
+	return AddEntry(FEditorMenuEntry::InitMenuSeparator(InName));
+}
+
+int32 FEditorMenuSection::IndexOfBlock(const FName InName) const
+{
+	for (int32 i=0; i < Blocks.Num(); ++i)
+	{
+		if (Blocks[i].Name == InName)
+		{
+			return i;
+		}
+	}
+
+	return INDEX_NONE;
+}
+
+void FEditorMenuSection::AssembleBlock(const FEditorMenuEntry& Block)
+{
+	const EEditorMenuInsertType Position = Block.InsertPosition.Position;
+
+	int32 ExistingIndex = IndexOfBlock(Block.Name);
+	if (ExistingIndex != INDEX_NONE)
+	{
+		Blocks[ExistingIndex] = Block;
+	}
+	else if (Position == EEditorMenuInsertType::Before || Position == EEditorMenuInsertType::After)
+	{
+		int32 DestIndex = IndexOfBlock(Block.InsertPosition.Name);
+		if (DestIndex != INDEX_NONE)
+		{
+			if (Position == EEditorMenuInsertType::After)
+			{
+				++DestIndex;
+			}
+
+			Blocks.Insert(Block, DestIndex);
+		}
+	}
+	else if (Position == EEditorMenuInsertType::First)
+	{
+		Blocks.Insert(Block, 0);
+	}
+	else
+	{
+		Blocks.Add(Block);
+	}
+}
+
+
+int32 FEditorMenuSection::RemoveEntry(const FName InName)
+{
+	return Blocks.RemoveAll([InName](const FEditorMenuEntry& Block) { return Block.Name == InName; });
+}
+
+int32 FEditorMenuSection::RemoveEntriesByOwner(const FEditorMenuOwner InOwner)
+{
+	if (InOwner != FEditorMenuOwner())
+	{
+		return Blocks.RemoveAll([InOwner](const FEditorMenuEntry& Block) { return Block.Owner == InOwner; });
+	}
+
+	return 0;
+}
+int32 FEditorMenuSection::FindBlockInsertIndex(const FEditorMenuEntry& InBlock) const
+{
+	const FEditorMenuInsert InPosition = InBlock.InsertPosition;
+
+	if (InPosition.IsDefault())
+	{
+		return Blocks.Num();
+	}
+
+	if (InPosition.Position == EEditorMenuInsertType::First)
+	{
+		for (int32 i = 0; i < Blocks.Num(); ++i)
+		{
+			if (Blocks[i].InsertPosition != InPosition)
+			{
+				return i;
+			}
+		}
+
+		return Blocks.Num();
+	}
+
+	int32 DestIndex = IndexOfBlock(InPosition.Name);
+	if (DestIndex == INDEX_NONE)
+	{
+		return INDEX_NONE;
+	}
+
+	if (InsertPosition.Position == EEditorMenuInsertType::After)
+	{
+		++DestIndex;
+	}
+
+	for (int32 i = DestIndex; i < Blocks.Num(); ++i)
+	{
+		if (Blocks[i].InsertPosition != InPosition)
+		{
+			return i;
+		}
+	}
+
+	return Blocks.Num();
+}
diff --git a/Engine/Source/Editor/EditorMenus/Private/EditorMenuSubsystem.cpp b/Engine/Source/Editor/EditorMenus/Private/EditorMenuSubsystem.cpp
new file mode 100644
--- /dev/null
+++ b/Engine/Source/Editor/EditorMenus/Private/EditorMenuSubsystem.cpp
@@ -0,0 +1,1216 @@
+// Copyright 1998-2019 Epic Games, Inc. All Rights Reserved.
+
+#include "EditorMenuSubsystem.h"
+#include "IEditorMenusModule.h"
+
+#include "Textures/SlateIcon.h"
+#include "Framework/Commands/UIAction.h"
+#include "Framework/MultiBox/MultiBox.h"
+#include "Framework/MultiBox/MultiBoxBuilder.h"
+#include "Internationalization/Internationalization.h"
+#include "Editor/EditorPerProjectUserSettings.h"
+
+#include "Editor.h"
+
+#define LOCTEXT_NAMESPACE "EditorMenuSubsystem"
+
+FAutoConsoleCommand EditorMenusRefreshMenuWidget = FAutoConsoleCommand(
+	TEXT("EditorMenus.RefreshAllWidgets"),
+	TEXT("Refresh All Editor Menu Widgets"),
+	FConsoleCommandDelegate::CreateLambda([]() {
+		UEditorMenuSubsystem::Get()->RefreshAllWidgets();
+	}));
+
+FEditorUIActionChoice::FEditorUIActionChoice(const TSharedPtr< const FUICommandInfo >& InCommand, const FUICommandList& InCommandList)
+{
+	if (InCommand.IsValid())
+	{
+		if (const FUIAction* UIAction = InCommandList.GetActionForCommand(InCommand))
+		{
+			Action = *UIAction;
+			EditorAction.Reset();
+			DynamicEditorAction.Reset();
+		}
+	}
+}
+
+UEditorMenuSubsystem::UEditorMenuSubsystem() :
+	UEditorSubsystem(),
+	bRefreshWidgetsNextTick(false)
+{
+}
+
+void UEditorMenuSubsystem::Initialize(FSubsystemCollectionBase& Collection)
+{
+
+}
+
+void UEditorMenuSubsystem::Deinitialize()
+{
+	bRefreshWidgetsNextTick = false;
+	GeneratedMenuWidgets.Reset();
+	WidgetObjectReferences.Reset();
+	Menus.Reset();
+}
+
+FName UEditorMenuSubsystem::JoinMenuPaths(const FName Base, const FName Child)
+{
+	return *(Base.ToString() + TEXT(".") + Child.ToString());
+}
+
+bool UEditorMenuSubsystem::GetDisplayUIExtensionPoints()
+{
+	return GetDefault<UEditorPerProjectUserSettings>()->bDisplayUIExtensionPoints;
+}
+
+UEditorMenu* UEditorMenuSubsystem::FindMenu(const FName Name)
+{
+	UEditorMenu** Found = Menus.Find(Name);
+	return Found ? *Found : nullptr;
+}
+
+TArray<UEditorMenu*> UEditorMenuSubsystem::CollectHierarchy(const FName InName)
+{
+	TArray<UEditorMenu*> Result;
+
+	UEditorMenu* Current = FindMenu(InName);
+	while (Current)
+	{
+		// Detect infinite loop
+		for (UEditorMenu* Other : Result)
+		{
+			if (Other->MenuName == Current->MenuName)
+			{
+				UE_LOG(LogEditorMenus, Warning, TEXT("Infinite loop detected in editor menu: %s"), *InName.ToString());
+				return TArray<UEditorMenu*>();
+			}
+		}
+
+		Result.Add(Current);
+
+		if (Current->MenuParent != NAME_None)
+		{
+			Current = FindMenu(Current->MenuParent);
+		}
+		else
+		{
+			break;
+		}
+	}
+
+	Algo::Reverse(Result);
+
+	return Result;
+}
+
+void UEditorMenuSubsystem::ListAllParents(const FName InName, TArray<FName>& AllParents)
+{
+	for (const UEditorMenu* Menu : CollectHierarchy(InName))
+	{
+		AllParents.Add(Menu->MenuName);
+	}
+}
+
+void UEditorMenuSubsystem::AssembleMenuByName(UEditorMenu* GeneratedMenu, const FName InName)
+{
+	if (UEditorMenu* Menu = FindMenu(InName))
+	{
+		GeneratedMenu->InitGeneratedCopy(Menu);
+		AssembleMenuHierarchy(GeneratedMenu, CollectHierarchy(InName));
+	}
+}
+
+void UEditorMenuSubsystem::AssembleMenuSection(UEditorMenu* GeneratedMenu, const UEditorMenu* Other, FEditorMenuSection* DestSection, const FEditorMenuSection& OtherSection)
+{
+	// Build list of blocks in expected order including blocks created by construct delegates
+	TArray<FEditorMenuEntry> RemainingBlocks;
+	for (const FEditorMenuEntry& Block : OtherSection.Blocks)
+	{
+		if (Block.IsScriptObjectDynamicConstruct() || Block.Construct.IsBound())
+		{
+			UEditorMenu* Constructed = NewObject<UEditorMenu>(this);
+			Constructed->Context = DestSection->Context;
+
+			if (Block.IsScriptObjectDynamicConstruct())
+			{
+				Block.ScriptObject->ConstructMenuEntry(*DestSection, DestSection->Context);
+			}
+			else
+			{
+				FEditorMenuSection& ConstructedSection = Constructed->AddSection(DestSection->Name);
+				Block.Construct.Execute(ConstructedSection);
+			}
+
+			// Combine all user's choice of selections here into the current section target
+			// If the user wants to add items to different sections they will need to create dynamic section instead (for now)
+			for (FEditorMenuSection& ConstructedSection : Constructed->Sections)
+			{
+				for (FEditorMenuEntry& ConstructedBlock : ConstructedSection.Blocks)
+				{
+					if (ConstructedBlock.InsertPosition.IsDefault())
+					{
+						ConstructedBlock.InsertPosition = Block.InsertPosition;
+					}
+					RemainingBlocks.Add(ConstructedBlock);
+				}
+			}
+		}
+		else
+		{
+			RemainingBlocks.Add(Block);
+		}
+	}
+
+	// Repeatedly loop because insert location may not exist until later in list
+	while (RemainingBlocks.Num() > 0)
+	{
+		int32 NumHandled = 0;
+		for (int32 i = 0; i < RemainingBlocks.Num(); ++i)
+		{
+			FEditorMenuEntry& Block = RemainingBlocks[i];
+			int32 DestIndex = DestSection->FindBlockInsertIndex(Block);
+			if (DestIndex != INDEX_NONE)
+			{
+				DestSection->Blocks.Insert(Block, DestIndex);
+				RemainingBlocks.RemoveAt(i);
+				--i;
+				++NumHandled;
+				// Restart loop because items earlier in the list may need to attach to this block
+				break;
+			}
+		}
+		if (NumHandled == 0)
+		{
+			for (const FEditorMenuEntry& Block : RemainingBlocks)
+			{
+				UE_LOG(LogEditorMenus, Warning, TEXT("Menu item not found: '%s' for insert: '%s'"), *Block.InsertPosition.Name.ToString(), *Block.Name.ToString());
+			}
+			break;
+		}
+	}
+}
+
+void UEditorMenuSubsystem::AssembleMenu(UEditorMenu* GeneratedMenu, const UEditorMenu* Other)
+{
+	TArray<FEditorMenuSection> RemainingSections;
+	for (const FEditorMenuSection& OtherSection : Other->Sections)
+	{
+		if (OtherSection.EditorMenuSectionDynamic || OtherSection.Construct.NewEditorMenuDelegate.IsBound())
+		{
+			UEditorMenu* ConstructedSections = NewObject<UEditorMenu>(this);
+			ConstructedSections->Context = GeneratedMenu->Context;
+
+			if (OtherSection.EditorMenuSectionDynamic)
+			{
+				OtherSection.EditorMenuSectionDynamic->ConstructSections(ConstructedSections, GeneratedMenu->Context);
+			}
+			else if (OtherSection.Construct.NewEditorMenuDelegate.IsBound())
+			{
+				OtherSection.Construct.NewEditorMenuDelegate.Execute(ConstructedSections);
+			}
+
+			for (FEditorMenuSection& ConstructedSection : ConstructedSections->Sections)
+			{
+				if (ConstructedSection.InsertPosition.IsDefault())
+				{
+					ConstructedSection.InsertPosition = OtherSection.InsertPosition;
+				}
+				RemainingSections.Add(ConstructedSection);
+			}
+		}
+		else
+		{
+			RemainingSections.Add(OtherSection);
+		}
+	}
+
+	while (RemainingSections.Num() > 0)
+	{
+		int32 NumHandled = 0;
+		for (int32 i=0; i < RemainingSections.Num(); ++i)
+		{
+			const FEditorMenuSection& RemainingSection = RemainingSections[i];
+
+			// Update existing section
+			FEditorMenuSection* Section = GeneratedMenu->FindSection(RemainingSection.Name);
+			if (!Section)
+			{
+				// Try add new section (if insert location exists)
+				int32 DestIndex = GeneratedMenu->FindInsertIndex(RemainingSection);
+				if (DestIndex != INDEX_NONE)
+				{
+					GeneratedMenu->Sections.InsertDefaulted(DestIndex);
+					Section = &GeneratedMenu->Sections[DestIndex];
+					Section->InitGeneratedSectionCopy(RemainingSection, GeneratedMenu->Context);
+				}
+				else
+				{
+					continue;
+				}
+			}
+
+			AssembleMenuSection(GeneratedMenu, Other, Section, RemainingSection);
+			RemainingSections.RemoveAt(i);
+			--i;
+			++NumHandled;
+			break;
+		}
+		if (NumHandled == 0)
+		{
+			for (const FEditorMenuSection& RemainingSection : RemainingSections)
+			{
+				UE_LOG(LogEditorMenus, Warning, TEXT("Menu section not found: '%s' for insert: '%s'"), *RemainingSection.InsertPosition.Name.ToString(), *RemainingSection.Name.ToString());
+			}
+			break;
+		}
+	}
+}
+
+int32 UEditorMenuSubsystem::FindCustomizedMenuIndex(const FName InName)
+{
+	for (int32 i = 0; i < CustomizedMenus.Num(); ++i)
+	{
+		if (CustomizedMenus[i].Name == InName)
+		{
+			return i;
+		}
+	}
+
+	return INDEX_NONE;
+}
+
+FCustomizedEditorMenu* UEditorMenuSubsystem::FindCustomizedMenu(const FName InName)
+{
+	int32 FoundIndex = FindCustomizedMenuIndex(InName);
+	return (FoundIndex != INDEX_NONE) ? &CustomizedMenus[FoundIndex] : nullptr;
+}
+
+void UEditorMenuSubsystem::ApplyCustomization(UEditorMenu* GeneratedMenu)
+{
+	const FCustomizedEditorMenu* CustomizedMenu = FindCustomizedMenu(GeneratedMenu->MenuName);
+	if (CustomizedMenu == nullptr)
+	{
+		return;
+	}
+
+	TArray<FEditorMenuSection> NewSections;
+	NewSections.Reserve(GeneratedMenu->Sections.Num());
+
+	for (const FCustomizedEditorMenuSection& CustomizedSection : CustomizedMenu->Sections)
+	{
+		int32 SectionIndex = GeneratedMenu->IndexOfSection(CustomizedSection.Name);
+		if (SectionIndex == INDEX_NONE)
+		{
+			continue;
+		}
+
+		FEditorMenuSection& Section = GeneratedMenu->Sections[SectionIndex];
+
+		TArray<FEditorMenuEntry> NewBlocks;
+		NewBlocks.Reserve(Section.Blocks.Num());
+
+		for (const FName ItemName : CustomizedSection.Items)
+		{
+			int32 EntrySectionIndex = INDEX_NONE;
+			int32 EntryIndex = INDEX_NONE;
+			if (GeneratedMenu->FindEntry(ItemName, EntrySectionIndex, EntryIndex))
+			{
+				NewBlocks.Add(GeneratedMenu->Sections[EntrySectionIndex].Blocks[EntryIndex]);
+				GeneratedMenu->Sections[EntrySectionIndex].Blocks.RemoveAt(EntryIndex);
+			}
+		}
+
+		// Remaining blocks
+		for (FEditorMenuEntry& Block : Section.Blocks)
+		{
+			NewBlocks.Add(Block);
+		}
+
+		Section.Blocks = NewBlocks;
+
+		NewSections.Add(Section);
+		GeneratedMenu->Sections.RemoveAt(SectionIndex);
+	}
+
+	// Remaining sections
+	NewSections.Append(GeneratedMenu->Sections);
+
+	// Hide
+	for (int32 SectionIndex=0; SectionIndex < NewSections.Num(); ++SectionIndex)
+	{
+		FEditorMenuSection& Section = NewSections[SectionIndex];
+		if (CustomizedMenu->HiddenSections.Contains(Section.Name))
+		{
+			NewSections.RemoveAt(SectionIndex);
+			--SectionIndex;
+			continue;
+		}
+
+		for (int32 i = 0; i < Section.Blocks.Num(); ++i)
+		{
+			if (CustomizedMenu->HiddenItems.Contains(Section.Blocks[i].Name))
+			{
+				Section.Blocks.RemoveAt(i);
+				--i;
+			}
+		}
+	}
+
+	GeneratedMenu->Sections = NewSections;
+}
+
+void UEditorMenuSubsystem::AssembleMenuHierarchy(UEditorMenu* GeneratedMenu, const TArray<UEditorMenu*>& Hierarchy)
+{
+	if (GeneratedMenu->MenuType == EMultiBoxType::MenuBar)
+	{
+		// Menu Bars require one section
+		if (GeneratedMenu->Sections.Num() == 0)
+		{
+			GeneratedMenu->Sections.AddDefaulted();
+		}
+
+		FEditorMenuSection& MenuBarSection = GeneratedMenu->Sections[0];
+		for (const UEditorMenu* MenuData : Hierarchy)
+		{
+			for (const FEditorMenuSection& Section : MenuData->Sections)
+			{
+				for (const FEditorMenuEntry& Block : Section.Blocks)
+				{
+					MenuBarSection.AssembleBlock(Block);
+				}
+			}
+		}
+	}
+	else
+	{
+		for (const UEditorMenu* FoundParent : Hierarchy)
+		{
+			AssembleMenu(GeneratedMenu, FoundParent);
+		}
+	}
+
+	ApplyCustomization(GeneratedMenu);
+}
+
+void UEditorMenuSubsystem::FillMenuDynamic(FMenuBuilder& MenuBuilder, FNewEditorMenuDelegate InConstructMenu)
+{
+	if (InConstructMenu.IsBound())
+	{
+		// Create final menu
+		UEditorMenu* MenuData = NewObject<UEditorMenu>();
+		InConstructMenu.Execute(MenuData);
+
+		// Populate menu builder with final menu
+		PopulateMenuBuilder(MenuBuilder, MenuData);
+	}
+}
+
+void UEditorMenuSubsystem::FillMenu(class FMenuBuilder& MenuBuilder, FName InMenuName, FEditorMenuContext InMenuContext)
+{
+	// Create combined final menu
+	UEditorMenu* GeneratedMenu = NewObject<UEditorMenu>();
+	GeneratedMenu->Context = InMenuContext;
+	AssembleMenuByName(GeneratedMenu, InMenuName);
+
+	// Populate menu builder with final menu
+	PopulateMenuBuilder(MenuBuilder, GeneratedMenu);
+}
+
+TSharedRef<SWidget> UEditorMenuSubsystem::GenerateToolbarComboButtonMenu(const FName SubMenuFullName, FEditorMenuContext InContext)
+{
+	return GenerateWidget(SubMenuFullName, InContext);
+}
+
+void UEditorMenuSubsystem::FillMenuBarDropDown(class FMenuBuilder& MenuBuilder, FName InParentName, FName InChildName, FEditorMenuContext InMenuContext)
+{
+	if (UEditorMenu* MenuToUse = FindSubMenuToGenerateWith(InParentName, InChildName))
+	{
+		// Create combined final menu
+		UEditorMenu* GeneratedMenu = NewObject<UEditorMenu>();
+		GeneratedMenu->Context = InMenuContext;
+		AssembleMenuByName(GeneratedMenu, MenuToUse->MenuName);
+
+		// Populate menu builder with final menu
+		PopulateMenuBuilder(MenuBuilder, GeneratedMenu);
+	}
+}
+
+void UEditorMenuSubsystem::PopulateMenuBuilder(FMenuBuilder& MenuBuilder, UEditorMenu* MenuData)
+{
+	if (GetDisplayUIExtensionPoints())
+	{
+		MenuBuilder.AddMenuEntry(FText::FromName(MenuData->GetMenuName()), TAttribute<FText>(), FSlateIcon(), FUIAction(), "MenuName");
+	}
+
+	for (FEditorMenuSection& Section : MenuData->Sections)
+	{
+		if (Section.Construct.NewEditorMenuDelegateLegacy.IsBound())
+		{
+			Section.Construct.NewEditorMenuDelegateLegacy.Execute(MenuBuilder, MenuData);
+			continue;
+		}
+
+		MenuBuilder.BeginSection(Section.Name, Section.Label);
+
+		for (FEditorMenuEntry& Block : Section.Blocks)
+		{
+			if (Block.ConstructLegacy.IsBound())
+			{
+				Block.ConstructLegacy.Execute(MenuBuilder, MenuData);
+			}
+			else if (Block.Type == EMultiBlockType::MenuEntry)
+			{
+				if (Block.SubMenu.bIsSubMenu)
+				{
+					if (Block.SubMenu.Construct.NewMenuDelegate.IsBound())
+					{
+						MenuBuilder.AddSubMenu(
+							Block.Label,
+							Block.ToolTip,
+							Block.SubMenu.Construct.NewMenuDelegate,
+							Block.SubMenu.bOpenSubMenuOnClick,
+							Block.Icon.Get(),
+							Block.bShouldCloseWindowAfterMenuSelection,
+							Block.Name
+						);
+					}
+					else if (Block.SubMenu.Construct.NewEditorMenuDelegate.IsBound())
+					{
+						// SubMenu constructed each time it is opened
+						MenuBuilder.AddSubMenu(
+							Block.Label,
+							Block.ToolTip,
+							FNewMenuDelegate::CreateUObject(this, &UEditorMenuSubsystem::FillMenuDynamic, Block.SubMenu.Construct.NewEditorMenuDelegate),
+							Block.SubMenu.bOpenSubMenuOnClick,
+							Block.Icon.Get(),
+							Block.bShouldCloseWindowAfterMenuSelection,
+							Block.Name
+						);
+					}
+					else
+					{
+						// SubMenu registered once by name in database
+						FName SubMenuFullName = JoinMenuPaths(MenuData->MenuName, Block.Name);
+						MenuBuilder.AddSubMenu(
+							Block.Label,
+							Block.ToolTip,
+							FNewMenuDelegate::CreateUObject(this, &UEditorMenuSubsystem::FillMenu, SubMenuFullName, MenuData->Context),
+							Block.SubMenu.bOpenSubMenuOnClick,
+							Block.Icon.Get(),
+							Block.bShouldCloseWindowAfterMenuSelection,
+							Block.Name
+						);
+					}
+				}
+				else
+				{
+					if (Block.StringCommand.String.Len() > 0)
+					{
+						FUIAction UIAction = ConvertUIAction(Block.StringCommand, MenuData->Context);
+						MenuBuilder.AddMenuEntry(Block.Label, Block.ToolTip, Block.Icon.Get(), UIAction, Block.Name, Block.UserInterfaceActionType, Block.TutorialHighlightName);
+					}
+					else if (Block.Command.IsValid())
+					{
+						MenuBuilder.AddMenuEntry(Block.Command, Block.Name, Block.Label, Block.ToolTip, Block.Icon.Get());
+					}
+					else if (Block.ScriptObject)
+					{
+						UEditorMenuEntryScript* ScriptObject = Block.ScriptObject;
+						FUIAction UIAction = ConvertUIAction(ScriptObject, MenuData->Context);
+						MenuBuilder.AddMenuEntry(ScriptObject->CreateLabelAttribute(MenuData->Context), ScriptObject->CreateToolTipAttribute(MenuData->Context), Block.Icon.Get(), UIAction, ScriptObject->Data.Name, Block.UserInterfaceActionType, Block.TutorialHighlightName);
+					}
+					else
+					{
+						FUIAction UIAction = ConvertUIAction(Block.Action, MenuData->Context);
+						MenuBuilder.AddMenuEntry(Block.Label, Block.ToolTip, Block.Icon.Get(), UIAction, Block.Name, Block.UserInterfaceActionType, Block.TutorialHighlightName);
+					}
+				}
+			}
+			else if (Block.Type == EMultiBlockType::MenuSeparator)
+			{
+				MenuBuilder.AddMenuSeparator(Block.Name);
+			}
+		}
+
+		MenuBuilder.EndSection();
+	}
+
+	AddReferencedContextObjects(MenuBuilder.GetMultiBox(), MenuData->Context);
+}
+
+void UEditorMenuSubsystem::PopulateToolBarBuilder(FToolBarBuilder& ToolBarBuilder, UEditorMenu* MenuData)
+{
+	for (FEditorMenuSection& Section : MenuData->Sections)
+	{
+		if (Section.Construct.NewToolBarDelegateLegacy.IsBound())
+		{
+			Section.Construct.NewToolBarDelegateLegacy.Execute(ToolBarBuilder, MenuData);
+			continue;
+		}
+
+		ToolBarBuilder.BeginSection(Section.Name);
+
+		for (FEditorMenuEntry& Block : Section.Blocks)
+		{
+			if (Block.ToolBar.ConstructLegacy.IsBound())
+			{
+				Block.ToolBar.ConstructLegacy.Execute(ToolBarBuilder, MenuData);
+			}
+			else if (Block.Type == EMultiBlockType::ToolBarButton)
+			{
+				if (Block.StringCommand.String.Len() > 0)
+				{
+					FUIAction UIAction = ConvertUIAction(Block.StringCommand, MenuData->Context);
+					ToolBarBuilder.AddToolBarButton(UIAction, Block.Name, Block.Label, Block.ToolTip, Block.Icon, Block.UserInterfaceActionType, Block.TutorialHighlightName);
+				}
+				else if (Block.Command.IsValid())
+				{
+					ToolBarBuilder.AddToolBarButton(Block.Command, Block.Name, Block.Label, Block.ToolTip, Block.Icon, Block.TutorialHighlightName);
+				}
+				else if (Block.ScriptObject)
+				{
+					UEditorMenuEntryScript* ScriptObject = Block.ScriptObject;
+					FUIAction UIAction = ConvertUIAction(ScriptObject, MenuData->Context);
+					ToolBarBuilder.AddToolBarButton(UIAction, ScriptObject->Data.Name, ScriptObject->CreateLabelAttribute(MenuData->Context), ScriptObject->CreateToolTipAttribute(MenuData->Context), Block.Icon, Block.UserInterfaceActionType, Block.TutorialHighlightName);
+				}
+				else
+				{
+					FUIAction UIAction = ConvertUIAction(Block.Action, MenuData->Context);
+					ToolBarBuilder.AddToolBarButton(UIAction, Block.Name, Block.Label, Block.ToolTip, Block.Icon, Block.UserInterfaceActionType, Block.TutorialHighlightName);
+				}
+			}
+			else if (Block.Type == EMultiBlockType::ToolBarComboButton)
+			{
+				if (FUIAction* UIAction = Block.Action.GetUIAction())
+				{
+					ToolBarBuilder.AddComboButton(*UIAction, Block.ToolBar.ComboButtonContextMenuGenerator, Block.Label, Block.ToolTip, Block.Icon, Block.ToolBar.bSimpleComboBox, Block.TutorialHighlightName);
+				}
+				else
+				{
+					FName SubMenuFullName = JoinMenuPaths(MenuData->MenuName, Block.Name);
+					FOnGetContent Delegate = FOnGetContent::CreateUObject(this, &UEditorMenuSubsystem::GenerateToolbarComboButtonMenu, SubMenuFullName, MenuData->Context);
+					ToolBarBuilder.AddComboButton(FUIAction(), Delegate, Block.Label, Block.ToolTip, Block.Icon, Block.ToolBar.bSimpleComboBox, Block.TutorialHighlightName);
+				}
+			}
+			else if (Block.Type == EMultiBlockType::ToolBarSeparator)
+			{
+				ToolBarBuilder.AddSeparator(Block.Name);
+			}
+			else
+			{
+				UE_LOG(LogEditorMenus, Warning, TEXT("Toolbar '%s', item '%s', has unknown type: %d"), *MenuData->MenuName.ToString(), *Block.Name.ToString(), Block.Type);
+			}
+		}
+
+		ToolBarBuilder.EndSection();
+	}
+
+	AddReferencedContextObjects(ToolBarBuilder.GetMultiBox(), MenuData->Context);
+}
+
+void UEditorMenuSubsystem::PopulateMenuBarBuilder(FMenuBarBuilder& MenuBarBuilder, UEditorMenu* GeneratedMenu)
+{
+	if (GeneratedMenu->Sections.Num() > 0)
+	{
+		for (FEditorMenuEntry& Block : GeneratedMenu->Sections[0].Blocks)
+		{
+			MenuBarBuilder.AddPullDownMenu(
+				Block.Label.Get(),
+				Block.ToolTip.Get(),
+				FNewMenuDelegate::CreateUObject(this, &UEditorMenuSubsystem::FillMenuBarDropDown, GeneratedMenu->MenuName, Block.Name, GeneratedMenu->Context),
+				Block.Name
+			);
+		}
+
+		AddReferencedContextObjects(MenuBarBuilder.GetMultiBox(), GeneratedMenu->Context);
+	}
+}
+
+FUIAction UEditorMenuSubsystem::ConvertUIAction(FEditorUIActionChoice& Choice, FEditorMenuContext& Context)
+{
+	if (FEditorUIAction* EditorAction = Choice.GetEditorUIAction())
+	{
+		return ConvertUIAction(*EditorAction, Context);
+	}
+	else if (FEditorDynamicUIAction* DynamicEditorAction = Choice.GetEditorDynamicUIAction())
+	{
+		return ConvertUIAction(*DynamicEditorAction, Context);
+	}
+	else if (FUIAction* Action = Choice.GetUIAction())
+	{
+		return *Action;
+	}
+
+	return FUIAction();
+}
+
+FUIAction UEditorMenuSubsystem::ConvertUIAction(FEditorUIAction& Actions, FEditorMenuContext& Context)
+{
+	FUIAction UIAction;
+
+	if (Actions.ExecuteAction.IsBound())
+	{
+		UIAction.ExecuteAction.BindLambda([DelegateToCall = Actions.ExecuteAction, Context]()
+		{
+			DelegateToCall.ExecuteIfBound(Context);
+		});
+	}
+
+	if (Actions.CanExecuteAction.IsBound())
+	{
+		UIAction.CanExecuteAction.BindLambda([DelegateToCall = Actions.CanExecuteAction, Context]()
+		{
+			return DelegateToCall.Execute(Context);
+		});
+	}
+
+	if (Actions.GetActionCheckState.IsBound())
+	{
+		UIAction.GetActionCheckState.BindLambda([DelegateToCall = Actions.GetActionCheckState, Context]()
+		{
+			return DelegateToCall.Execute(Context);
+		});
+	}
+
+	if (Actions.IsActionVisibleDelegate.IsBound())
+	{
+		UIAction.IsActionVisibleDelegate.BindLambda([DelegateToCall = Actions.IsActionVisibleDelegate, Context]()
+		{
+			return DelegateToCall.Execute(Context);
+		});
+	}
+
+	return UIAction;
+}
+
+FUIAction UEditorMenuSubsystem::ConvertUIAction(FEditorDynamicUIAction& Actions, FEditorMenuContext& Context)
+{
+	FUIAction UIAction;
+
+	if (Actions.ExecuteAction.IsBound())
+	{
+		UIAction.ExecuteAction.BindLambda([DelegateToCall = Actions.ExecuteAction, Context]()
+		{
+			DelegateToCall.ExecuteIfBound(Context);
+		});
+	}
+
+	if (Actions.CanExecuteAction.IsBound())
+	{
+		UIAction.CanExecuteAction.BindLambda([DelegateToCall = Actions.CanExecuteAction, Context]()
+		{
+			return DelegateToCall.Execute(Context);
+		});
+	}
+
+	if (Actions.GetActionCheckState.IsBound())
+	{
+		UIAction.GetActionCheckState.BindLambda([DelegateToCall = Actions.GetActionCheckState, Context]()
+		{
+			return DelegateToCall.Execute(Context);
+		});
+	}
+
+	if (Actions.IsActionVisibleDelegate.IsBound())
+	{
+		UIAction.IsActionVisibleDelegate.BindLambda([DelegateToCall = Actions.IsActionVisibleDelegate, Context]()
+		{
+			return DelegateToCall.Execute(Context);
+		});
+	}
+
+	return UIAction;
+}
+
+FUIAction UEditorMenuSubsystem::ConvertUIAction(UEditorMenuEntryScript* ScriptObject, FEditorMenuContext& Context)
+{
+	FUIAction UIAction;
+
+	if (ScriptObject)
+	{
+		UClass* ScriptClass = ScriptObject->GetClass();
+
+		static const FName ExecuteName = GET_FUNCTION_NAME_CHECKED(UEditorMenuEntryScript, Execute);
+		if (ScriptClass->IsFunctionImplementedInScript(ExecuteName))
+		{
+			UIAction.ExecuteAction.BindUFunction(ScriptObject, ExecuteName, Context);
+		}
+
+		static const FName CanExecuteName = GET_FUNCTION_NAME_CHECKED(UEditorMenuEntryScript, CanExecute);
+		if (ScriptClass->IsFunctionImplementedInScript(CanExecuteName))
+		{
+			UIAction.CanExecuteAction.BindUFunction(ScriptObject, CanExecuteName, Context);
+		}
+
+		static const FName GetCheckStateName = GET_FUNCTION_NAME_CHECKED(UEditorMenuEntryScript, GetCheckState);
+		if (ScriptClass->IsFunctionImplementedInScript(GetCheckStateName))
+		{
+			UIAction.GetActionCheckState.BindUFunction(ScriptObject, GetCheckStateName, Context);
+		}
+
+		static const FName IsVisibleName = GET_FUNCTION_NAME_CHECKED(UEditorMenuEntryScript, IsVisible);
+		if (ScriptClass->IsFunctionImplementedInScript(IsVisibleName))
+		{
+			UIAction.IsActionVisibleDelegate.BindUFunction(ScriptObject, IsVisibleName, Context);
+		}
+	}
+
+	return UIAction;
+}
+
+void UEditorMenuSubsystem::ExecuteStringCommand(const FEditorMenuStringCommand StringCommand, const FEditorMenuContext Context)
+{
+	if (StringCommand.String.Len() > 0)
+	{
+		static const FName CommandTypeName("Command");
+		if (const FEditorMenuExecuteString* Handler = StringCommandHandlers.Find(StringCommand.TypeName))
+		{
+			if (Handler->IsBound())
+			{
+				Handler->Execute(StringCommand.String, Context);
+			}
+		}
+		else if (StringCommand.TypeName == CommandTypeName)
+		{
+			if (GEditor)
+			{
+				GEditor->Exec(nullptr, *StringCommand.String);
+			}
+		}
+		else
+		{
+			UE_LOG(LogEditorMenus, Warning, TEXT("Unknown string command handler type: '%s'"), *StringCommand.TypeName.ToString());
+		}
+	}
+}
+
+FUIAction UEditorMenuSubsystem::ConvertUIAction(const FEditorMenuStringCommand& StringCommand, const FEditorMenuContext& Context) const
+{
+	FUIAction UIAction;
+
+	if (StringCommand.String.Len() > 0)
+	{
+		UIAction.ExecuteAction.BindUObject(this, &UEditorMenuSubsystem::ExecuteStringCommand, StringCommand, Context);
+	}
+
+	return UIAction;
+}
+
+UEditorMenu* UEditorMenuSubsystem::FindSubMenuToGenerateWith(const FName InParentName, const FName InChildName)
+{
+	FName BaseName = InParentName;
+	while (BaseName != NAME_None)
+	{
+		FName JoinedName = JoinMenuPaths(BaseName, InChildName);
+		if (UEditorMenu* Found = FindMenu(JoinedName))
+		{
+			return Found;
+		}
+
+		UEditorMenu* BaseData = FindMenu(BaseName);
+		BaseName = BaseData ? BaseData->MenuParent : NAME_None;
+	}
+
+	return nullptr;
+}
+
+UObject* UEditorMenuSubsystem::FindContext(const FEditorMenuContext& InContext, UClass* InClass)
+{
+	return InContext.FindByClass(InClass);
+}
+
+void UEditorMenuSubsystem::AddReferencedContextObjects(const TSharedRef<FMultiBox>& InMultiBox, const FEditorMenuContext& InMenuContext)
+{
+	if (InMenuContext.ContextObjects.Num() == 0)
+	{
+		return;
+	}
+
+	TArray<UObject*>& References = WidgetObjectReferences.FindOrAdd(InMultiBox);
+	for (const TWeakObjectPtr<UObject>& WeakObject : InMenuContext.ContextObjects)
+	{
+		if (UObject* Object = WeakObject.Get())
+		{
+			References.AddUnique(Object);
+		}
+	}
+}
+
+void UEditorMenuSubsystem::AddReferencedObjects(UObject* InThis, FReferenceCollector& Collector)
+{
+	UEditorMenuSubsystem* This = CastChecked<UEditorMenuSubsystem>(InThis);
+
+	for (auto It = This->WidgetObjectReferences.CreateIterator(); It; ++It)
+	{
+		if (It->Key.IsValid())
+		{
+			Collector.AddReferencedObjects(It->Value, InThis);
+		}
+		else
+		{
+			It.RemoveCurrent();
+		}
+	}
+
+	Super::AddReferencedObjects(InThis, Collector);
+}
+
+UEditorMenu* UEditorMenuSubsystem::GenerateMenu(const FName Name, FEditorMenuContext& InMenuContext)
+{
+	return GenerateMenu(CollectHierarchy(Name), InMenuContext);
+}
+
+UEditorMenu* UEditorMenuSubsystem::GenerateMenu(const TArray<UEditorMenu*>& Hierarchy, FEditorMenuContext& InMenuContext)
+{
+	UEditorMenu* GeneratedMenu = NewObject<UEditorMenu>(this);
+
+	if (Hierarchy.Num() > 0)
+	{
+		GeneratedMenu->InitGeneratedCopy(Hierarchy[0]);
+		GeneratedMenu->MenuName = Hierarchy.Last()->MenuName;
+		GeneratedMenu->Context = InMenuContext;
+		AssembleMenuHierarchy(GeneratedMenu, Hierarchy);
+	}
+
+	return GeneratedMenu;
+}
+
+TSharedRef< class SWidget > UEditorMenuSubsystem::GenerateWidget(const FName InName, FEditorMenuContext& InMenuContext)
+{
+	UEditorMenu* Generated = GenerateMenu(InName, InMenuContext);
+	return GenerateWidget(Generated);
+}
+
+TSharedRef<SWidget> UEditorMenuSubsystem::GenerateWidget(const TArray<UEditorMenu*>& Hierarchy, FEditorMenuContext& InMenuContext)
+{
+	if (Hierarchy.Num() == 0)
+	{
+		return SNullWidget::NullWidget;
+	}
+
+	UEditorMenu* Generated = GenerateMenu(Hierarchy, InMenuContext);
+	return GenerateWidget(Generated);
+}
+
+TSharedRef<SWidget> UEditorMenuSubsystem::GenerateWidget(UEditorMenu* GeneratedMenu)
+{
+	FGeneratedEditorMenuWidgets& WidgetsForMenuName = GeneratedMenuWidgets.FindOrAdd(GeneratedMenu->MenuName);
+
+	// Store a copy so that we can call 'Refresh' on menus not in the database
+	FGeneratedEditorMenuWidget& GeneratedMenuWidget = WidgetsForMenuName.Instances.AddDefaulted_GetRef();
+	GeneratedMenuWidget.GeneratedMenu = DuplicateObject<UEditorMenu>(GeneratedMenu, this);
+	// Copy native properties that serialize does not
+	GeneratedMenuWidget.GeneratedMenu->Context = GeneratedMenu->Context;
+	GeneratedMenuWidget.GeneratedMenu->StyleSet = GeneratedMenu->StyleSet;
+
+	if (GeneratedMenu->MenuType == EMultiBoxType::Menu)
+	{
+		FMenuBuilder MenuBuilder(GeneratedMenu->bShouldCloseWindowAfterMenuSelection, GeneratedMenu->Context.CommandList, GeneratedMenu->Context.GetAllExtenders(), GeneratedMenu->bCloseSelfOnly, GeneratedMenu->StyleSet, GeneratedMenu->bSearchable);
+		PopulateMenuBuilder(MenuBuilder, GeneratedMenu);
+		TSharedRef<SWidget> Result = MenuBuilder.MakeWidget();
+		GeneratedMenuWidget.Widget = Result;
+		return Result;
+	}
+	else if (GeneratedMenu->MenuType == EMultiBoxType::MenuBar)
+	{
+		FMenuBarBuilder MenuBarBuilder(GeneratedMenu->Context.CommandList, GeneratedMenu->Context.GetAllExtenders(), GeneratedMenu->StyleSet);
+		PopulateMenuBarBuilder(MenuBarBuilder, GeneratedMenu);
+		TSharedRef<SWidget> Result = MenuBarBuilder.MakeWidget();
+		GeneratedMenuWidget.Widget = Result;
+		return Result;
+	}
+	else if (GeneratedMenu->MenuType == EMultiBoxType::ToolBar)
+	{
+		const EOrientation ToolBarOrientation = (GeneratedMenu->MenuType == EMultiBoxType::VerticalToolBar) ? Orient_Vertical : Orient_Horizontal;
+		FToolBarBuilder ToolbarBuilder(GeneratedMenu->Context.CommandList, GeneratedMenu->MenuName, GeneratedMenu->Context.GetAllExtenders(), ToolBarOrientation, GeneratedMenu->bToolBarForceSmallIcons);
+		ToolbarBuilder.SetIsFocusable(GeneratedMenu->bToolBarIsFocusable);
+		PopulateToolBarBuilder(ToolbarBuilder, GeneratedMenu);
+		TSharedRef<SWidget> Result = ToolbarBuilder.MakeWidget();
+		GeneratedMenuWidget.Widget = Result;
+		return Result;
+	}
+
+	return SNullWidget::NullWidget;
+}
+
+void UEditorMenuSubsystem::RefreshAllWidgets()
+{
+	if (bRefreshWidgetsNextTick == false)
+	{
+		bRefreshWidgetsNextTick = true;
+		GEditor->GetTimerManager()->SetTimerForNextTick(this, &UEditorMenuSubsystem::HandleRefreshAllWidgetsNextTick);
+	}
+}
+
+void UEditorMenuSubsystem::CleanupStaleGeneratedMenus()
+{
+	for (auto WidgetsForMenuNameIt = GeneratedMenuWidgets.CreateIterator(); WidgetsForMenuNameIt; ++WidgetsForMenuNameIt)
+	{
+		FGeneratedEditorMenuWidgets& WidgetsForMenuName = WidgetsForMenuNameIt->Value;
+
+		for (auto Instance = WidgetsForMenuName.Instances.CreateIterator(); Instance; ++Instance)
+		{
+			if (!Instance->Widget.IsValid())
+			{
+				Instance.RemoveCurrent();
+			}
+		}
+
+		if (WidgetsForMenuName.Instances.Num() == 0)
+		{
+			WidgetsForMenuNameIt.RemoveCurrent();
+		}
+	}
+}
+
+void UEditorMenuSubsystem::HandleRefreshAllWidgetsNextTick()
+{
+	if (bRefreshWidgetsNextTick)
+	{
+		CleanupStaleGeneratedMenus();
+
+		for (auto WidgetsForMenuNameIt = GeneratedMenuWidgets.CreateIterator(); WidgetsForMenuNameIt; ++WidgetsForMenuNameIt)
+		{
+			FGeneratedEditorMenuWidgets& WidgetsForMenuName = WidgetsForMenuNameIt->Value;
+			for (auto Instance = WidgetsForMenuName.Instances.CreateIterator(); Instance; ++Instance)
+			{
+				if (Instance->Widget.IsValid())
+				{
+					RefreshMenuWidget(WidgetsForMenuNameIt->Key, *Instance);
+				}
+			}
+		}
+
+		bRefreshWidgetsNextTick = false;
+	}
+}
+
+bool UEditorMenuSubsystem::RefreshMenuWidget(const FName InName)
+{
+	if (FGeneratedEditorMenuWidgets* WidgetsForMenuName = GeneratedMenuWidgets.Find(InName))
+	{
+		for (auto Instance = WidgetsForMenuName->Instances.CreateIterator(); Instance; ++Instance)
+		{
+			if (RefreshMenuWidget(InName, *Instance))
+			{
+				return true;
+			}
+			else
+			{
+				Instance.RemoveCurrent();
+			}
+		}
+	}
+
+	return false;
+}
+
+bool UEditorMenuSubsystem::RefreshMenuWidget(const FName InName, FGeneratedEditorMenuWidget& GeneratedMenuWidget)
+{
+	if (!GeneratedMenuWidget.Widget.IsValid())
+	{
+		return false;
+	}
+
+	// Regenerate menu from database
+	UEditorMenu* GeneratedMenu = GenerateMenu(InName, GeneratedMenuWidget.GeneratedMenu->Context);
+	GeneratedMenuWidget.GeneratedMenu = GeneratedMenu;
+
+	// Regenerate Multibox
+	TSharedRef<SMultiBoxWidget> MultiBoxWidget = StaticCastSharedRef<SMultiBoxWidget>(GeneratedMenuWidget.Widget.Pin().ToSharedRef());
+	if (GeneratedMenu->MenuType == EMultiBoxType::Menu)
+	{
+		FMenuBuilder MenuBuilder(GeneratedMenu->bShouldCloseWindowAfterMenuSelection, GeneratedMenu->Context.CommandList, GeneratedMenu->Context.GetAllExtenders(), GeneratedMenu->bCloseSelfOnly, GeneratedMenu->StyleSet, GeneratedMenu->bSearchable);
+		PopulateMenuBuilder(MenuBuilder, GeneratedMenu);
+		MultiBoxWidget->SetMultiBox(MenuBuilder.GetMultiBox());
+	}
+	else if (GeneratedMenu->MenuType == EMultiBoxType::MenuBar)
+	{
+		FMenuBarBuilder MenuBarBuilder(GeneratedMenu->Context.CommandList, GeneratedMenu->Context.GetAllExtenders(), GeneratedMenu->StyleSet);
+		PopulateMenuBarBuilder(MenuBarBuilder, GeneratedMenu);
+		MultiBoxWidget->SetMultiBox(MenuBarBuilder.GetMultiBox());
+	}
+	else if (GeneratedMenu->MenuType == EMultiBoxType::ToolBar)
+	{
+		const EOrientation ToolBarOrientation = (GeneratedMenu->MenuType == EMultiBoxType::VerticalToolBar) ? Orient_Vertical : Orient_Horizontal;
+		FToolBarBuilder ToolbarBuilder(GeneratedMenu->Context.CommandList, GeneratedMenu->MenuName, GeneratedMenu->Context.GetAllExtenders(), ToolBarOrientation, GeneratedMenu->bToolBarForceSmallIcons);
+		ToolbarBuilder.SetIsFocusable(GeneratedMenu->bToolBarIsFocusable);
+		PopulateToolBarBuilder(ToolbarBuilder, GeneratedMenu);
+		MultiBoxWidget->SetMultiBox(ToolbarBuilder.GetMultiBox());
+	}
+
+	MultiBoxWidget->BuildMultiBoxWidget();
+	return true;
+}
+
+UEditorMenu* UEditorMenuSubsystem::GenerateMenuAsBuilder(const UEditorMenu* InMenu, FEditorMenuContext& InMenuContext)
+{
+	TArray<UEditorMenu*> Hierarchy = CollectHierarchy(InMenu->MenuName);
+
+	// Insert InMenu as second to last so items in InMenu appear before items registered in database by other plugins
+	if (Hierarchy.Num() > 0)
+	{
+		Hierarchy.Insert((UEditorMenu*)InMenu, Hierarchy.Num() - 1);
+	}
+	else
+	{
+		Hierarchy.Add((UEditorMenu*)InMenu);
+	}
+
+	return GenerateMenu(Hierarchy, InMenuContext);
+}
+
+UEditorMenu* UEditorMenuSubsystem::RegisterMenu(const FName InName, const FName InParent, EMultiBoxType InType)
+{
+	if (UEditorMenu* Found = FindMenu(InName))
+	{
+		if (!Found->bRegistered)
+		{
+			Found->MenuParent = InParent;
+			Found->MenuType = InType;
+			Found->MenuOwner = CurrentOwner();
+			Found->bRegistered = true;
+		}
+		else
+		{
+			UE_LOG(LogEditorMenus, Warning, TEXT("Menu already registered : %s"), *InName.ToString());
+		}
+
+		return Found;
+	}
+
+	UEditorMenu* EditorMenu = NewObject<UEditorMenu>(this);
+	EditorMenu->InitMenu(CurrentOwner(), InName, InParent, InType);
+	EditorMenu->bRegistered = true;
+	Menus.Add(InName, EditorMenu);
+	return EditorMenu;
+}
+
+UEditorMenu* UEditorMenuSubsystem::ExtendMenu(const FName InName)
+{
+	if (UEditorMenu* Found = FindMenu(InName))
+	{
+		return Found;
+	}
+
+	UEditorMenu* EditorMenu = NewObject<UEditorMenu>(this);
+	EditorMenu->MenuName = InName;
+	EditorMenu->bRegistered = false;
+	Menus.Add(InName, EditorMenu);
+	return EditorMenu;
+}
+
+void UEditorMenuSubsystem::RemoveMenu(const FName MenuName)
+{
+	Menus.Remove(MenuName);
+}
+
+bool UEditorMenuSubsystem::AddMenuEntryObject(UEditorMenuEntryScript* MenuEntryObject)
+{
+	UEditorMenu* Menu = UEditorMenuSubsystem::Get()->ExtendMenu(MenuEntryObject->Data.Menu);
+	Menu->AddMenuEntryObject(MenuEntryObject);
+	return true;
+}
+
+void UEditorMenuSubsystem::SetSectionLabel(const FName MenuName, const FName SectionName, const FText Label)
+{
+	ExtendMenu(MenuName)->FindOrAddSection(SectionName).Label = TAttribute<FText>(Label);
+}
+
+void UEditorMenuSubsystem::SetSectionPosition(const FName MenuName, const FName SectionName, const FName PositionName, const EEditorMenuInsertType PositionType)
+{
+	ExtendMenu(MenuName)->FindOrAddSection(SectionName).InsertPosition = FEditorMenuInsert(PositionName, PositionType);
+}
+
+void UEditorMenuSubsystem::AddSection(const FName MenuName, const FName SectionName, const TAttribute< FText >& InLabel, const FEditorMenuInsert InPosition)
+{
+	UEditorMenu* Menu = ExtendMenu(MenuName);
+	FEditorMenuSection* Section = Menu->FindSection(SectionName);
+	if (!Section)
+	{
+		Menu->AddSection(SectionName, InLabel, InPosition);
+	}
+}
+
+void UEditorMenuSubsystem::RemoveSection(const FName MenuName, const FName InSection)
+{
+	if (UEditorMenu* Menu = FindMenu(MenuName))
+	{
+		Menu->RemoveSection(InSection);
+	}
+}
+
+void UEditorMenuSubsystem::AddEntry(const FName MenuName, const FName InSection, const FEditorMenuEntry& InEntry)
+{
+	ExtendMenu(MenuName)->FindOrAddSection(InSection).AddEntry(InEntry);
+}
+
+void UEditorMenuSubsystem::RemoveEntry(const FName MenuName, const FName InSection, const FName InName)
+{
+	if (UEditorMenu* Menu = FindMenu(MenuName))
+	{
+		if (FEditorMenuSection* Section = Menu->FindSection(InSection))
+		{
+			Section->RemoveEntry(InName);
+		}
+	}
+}
+
+void UEditorMenuSubsystem::UnregisterOwnerInternal(FEditorMenuOwner InOwner)
+{
+	if (InOwner != FEditorMenuOwner())
+	{
+		for (auto It = Menus.CreateIterator(); It; ++It)
+		{
+			int32 NumEntriesRemoved = 0;
+
+			UEditorMenu* Menu = It->Value;
+			for (FEditorMenuSection& Section : Menu->Sections)
+			{
+				NumEntriesRemoved += Section.RemoveEntriesByOwner(InOwner);
+			}
+
+			// Refresh any widgets that are currently displayed to the user
+			if (NumEntriesRemoved > 0)
+			{
+				RefreshAllWidgets();
+			}
+		}
+	}
+}
+
+FEditorMenuOwner UEditorMenuSubsystem::CurrentOwner() const
+{
+	if (OwnerStack.Num() > 0)
+	{
+		return OwnerStack.Last();
+	}
+
+	return FEditorMenuOwner();
+}
+
+void UEditorMenuSubsystem::PushOwner(const FEditorMenuOwner InOwner)
+{
+	OwnerStack.Add(InOwner);
+}
+
+void UEditorMenuSubsystem::PopOwner(const FEditorMenuOwner InOwner)
+{
+	FEditorMenuOwner PoppedOwner = OwnerStack.Pop(false);
+	check(PoppedOwner == InOwner);
+}
+
+void UEditorMenuSubsystem::UnregisterOwnerByName(FName InOwnerName)
+{
+	UnregisterOwnerInternal(InOwnerName);
+}
+
+void UEditorMenuSubsystem::RegisterStringCommandHandler(const FName InName, const FEditorMenuExecuteString& InDelegate)
+{
+	StringCommandHandlers.Add(InName, InDelegate);
+}
+
+void UEditorMenuSubsystem::UnregisterStringCommandHandler(const FName InName)
+{
+	StringCommandHandlers.Remove(InName);
+}
diff --git a/Engine/Source/Editor/EditorMenus/Private/EditorMenusBlueprintLibrary.cpp b/Engine/Source/Editor/EditorMenus/Private/EditorMenusBlueprintLibrary.cpp
new file mode 100644
--- /dev/null
+++ b/Engine/Source/Editor/EditorMenus/Private/EditorMenusBlueprintLibrary.cpp
@@ -0,0 +1,90 @@
+// Copyright 1998-2019 Epic Games, Inc. All Rights Reserved.
+
+
+#include "EditorMenusBlueprintLibrary.h"
+
+FEditorMenuOwner UEditorMenuEntryExtensions::MakeEditorMenuOwner(FName Name)
+{
+	return FEditorMenuOwner(Name);
+}
+
+void UEditorMenuEntryExtensions::BreakEditorMenuOwner(const FEditorMenuOwner& InValue, FName& Name)
+{
+	Name = InValue.TryGetName();
+}
+
+UObject* UEditorMenuContextExtensions::FindByClass(const FEditorMenuContext& Context, UClass* InClass)
+{
+	return Context.FindByClass(InClass);
+}
+
+void UEditorMenuEntryExtensions::SetLabel(UPARAM(ref) FEditorMenuEntry& Target, const FText& Label)
+{
+	Target.Label = Label;
+}
+
+FText UEditorMenuEntryExtensions::GetLabel(const FEditorMenuEntry& Target)
+{
+	return Target.Label.Get();
+}
+
+void UEditorMenuEntryExtensions::SetToolTip(UPARAM(ref) FEditorMenuEntry& Target, const FText& ToolTip)
+{
+	Target.ToolTip = ToolTip;
+}
+
+FText UEditorMenuEntryExtensions::GetToolTip(const FEditorMenuEntry& Target)
+{
+	return Target.ToolTip.Get();
+}
+
+void UEditorMenuEntryExtensions::SetIcon(UPARAM(ref) FEditorMenuEntry& Target, const FName StyleSetName, const FName StyleName, const FName SmallStyleName)
+{
+	if (SmallStyleName == NAME_None)
+	{
+		if (StyleSetName == NAME_None && StyleName == NAME_None)
+		{
+			Target.Icon = FSlateIcon();
+		}
+			
+		Target.Icon = FSlateIcon(StyleSetName, StyleName);
+	}
+
+	Target.Icon = FSlateIcon(StyleSetName, StyleName, SmallStyleName);
+}
+
+void UEditorMenuEntryExtensions::SetStringCommand(UPARAM(ref) FEditorMenuEntry& Target, const FName TypeName, const FString& String)
+{
+	Target.ResetActions();
+	Target.StringCommand.TypeName = TypeName;
+	Target.StringCommand.String = String;
+}
+
+FEditorMenuEntry UEditorMenuEntryExtensions::InitMenuEntry(const FName InOwner, const FName InName, const FText& InLabel, const FText& InToolTip, const FEditorMenuStringCommand& StringCommand)
+{
+	FEditorMenuEntry Entry(InOwner, InName, EMultiBlockType::MenuEntry);
+	Entry.Label = InLabel;
+	Entry.ToolTip = InToolTip;
+	Entry.StringCommand = StringCommand;
+	return Entry;
+}
+
+void UEditorMenuSectionExtensions::SetLabel(UPARAM(ref) FEditorMenuSection& Section, const FText& Label)
+{
+	Section.Label = Label;
+}
+
+FText UEditorMenuSectionExtensions::GetLabel(const FEditorMenuSection& Section)
+{
+	return Section.Label.Get();
+}
+
+void UEditorMenuSectionExtensions::AddEntry(UPARAM(ref) FEditorMenuSection& Section, const FEditorMenuEntry& Args)
+{
+	Section.AddEntry(Args);
+}
+
+void UEditorMenuSectionExtensions::AddEntryObject(UPARAM(ref) FEditorMenuSection& Section, UEditorMenuEntryScript* InObject)
+{
+	Section.AddEntryObject(InObject);
+}
diff --git a/Engine/Source/Editor/EditorMenus/Private/EditorMenusBlueprintLibrary.h b/Engine/Source/Editor/EditorMenus/Private/EditorMenusBlueprintLibrary.h
new file mode 100644
--- /dev/null
+++ b/Engine/Source/Editor/EditorMenus/Private/EditorMenusBlueprintLibrary.h
@@ -0,0 +1,78 @@
+// Copyright 1998-2019 Epic Games, Inc. All Rights Reserved.
+
+#pragma once
+
+#include "CoreMinimal.h"
+
+#include "EditorMenuDelegates.h"
+#include "EditorMenuEntry.h"
+#include "EditorMenuSection.h"
+#include "Misc/Attribute.h"
+
+#include "EditorMenusBlueprintLibrary.generated.h"
+
+UCLASS()
+class EDITORMENUS_API UEditorMenuContextExtensions : public UBlueprintFunctionLibrary
+{
+	GENERATED_BODY()
+
+public:
+
+	UFUNCTION(BlueprintCallable, Category = "Editor UI", meta = (ScriptMethod))
+	static UObject* FindByClass(const FEditorMenuContext& Context, UClass* InClass);
+};
+
+UCLASS()
+class EDITORMENUS_API UEditorMenuEntryExtensions : public UBlueprintFunctionLibrary
+{
+	GENERATED_BODY()
+
+public:
+
+	UFUNCTION(BlueprintPure, Category = "Editor UI", meta = (Keywords = "construct build", NativeMakeFunc))
+	static FEditorMenuOwner MakeEditorMenuOwner(FName Name);
+
+	UFUNCTION(BlueprintPure, Category = "Editor UI", meta = (NativeBreakFunc))
+	static void BreakEditorMenuOwner(const FEditorMenuOwner& InValue, FName& Name);
+
+	UFUNCTION(BlueprintCallable, Category = "Editor UI", meta = (ScriptMethod))
+	static void SetLabel(UPARAM(ref) FEditorMenuEntry& Target, const FText& Label);
+
+	UFUNCTION(BlueprintCallable, Category = "Editor UI", meta = (ScriptMethod))
+	static FText GetLabel(const FEditorMenuEntry& Target);
+
+	UFUNCTION(BlueprintCallable, Category = "Editor UI", meta = (ScriptMethod))
+	static void SetToolTip(UPARAM(ref) FEditorMenuEntry& Target, const FText& ToolTip);
+
+	UFUNCTION(BlueprintCallable, Category = "Editor UI", meta = (ScriptMethod))
+	static FText GetToolTip(const FEditorMenuEntry& Target);
+
+	UFUNCTION(BlueprintCallable, Category = "Editor UI", meta = (ScriptMethod))
+	static void SetIcon(UPARAM(ref) FEditorMenuEntry& Target, const FName StyleSetName, const FName StyleName = NAME_None, const FName SmallStyleName = NAME_None);
+
+	UFUNCTION(BlueprintCallable, Category = "Editor UI", meta = (ScriptMethod))
+	static void SetStringCommand(UPARAM(ref) FEditorMenuEntry& Target, const FName TypeName, const FString& String);
+
+	UFUNCTION(BlueprintCallable, Category = "Editor UI")
+	static FEditorMenuEntry InitMenuEntry(const FName InOwner, const FName InName, const FText& InLabel, const FText& InToolTip, const FEditorMenuStringCommand& StringCommand);
+};
+
+UCLASS()
+class EDITORMENUS_API UEditorMenuSectionExtensions : public UBlueprintFunctionLibrary
+{
+	GENERATED_BODY()
+
+public:
+
+	UFUNCTION(BlueprintCallable, Category = "Editor UI", meta = (ScriptMethod))
+	static void SetLabel(UPARAM(ref) FEditorMenuSection& Section, const FText& Label);
+
+	UFUNCTION(BlueprintCallable, Category = "Editor UI", meta = (ScriptMethod))
+	static FText GetLabel(const FEditorMenuSection& Section);
+
+	UFUNCTION(BlueprintCallable, Category = "Editor UI", meta = (ScriptMethod))
+	static void AddEntry(UPARAM(ref) FEditorMenuSection& Section, const FEditorMenuEntry& Args);
+
+	UFUNCTION(BlueprintCallable, Category = "Editor UI", meta = (ScriptMethod))
+	static void AddEntryObject(UPARAM(ref) FEditorMenuSection& Section, UEditorMenuEntryScript* InObject);
+};
diff --git a/Engine/Source/Editor/EditorMenus/Private/EditorMenusModule.cpp b/Engine/Source/Editor/EditorMenus/Private/EditorMenusModule.cpp
new file mode 100644
--- /dev/null
+++ b/Engine/Source/Editor/EditorMenus/Private/EditorMenusModule.cpp
@@ -0,0 +1,30 @@
+// Copyright 1998-2019 Epic Games, Inc. All Rights Reserved.
+
+#include "CoreMinimal.h"
+#include "Modules/ModuleManager.h"
+#include "IEditorMenusModule.h"
+
+
+/**
+ * Implements the Editor menus module.
+ */
+class FEditorMenusModule
+	: public IEditorMenusModule
+{
+public:
+	
+	// IEditorMenusModule interface
+
+	virtual void StartupModule( ) override
+	{
+	}
+
+	virtual void ShutdownModule( ) override
+	{
+	}
+
+	// End IModuleInterface interface
+};
+
+
+IMPLEMENT_MODULE(FEditorMenusModule, EditorMenus)
diff --git a/Engine/Source/Editor/EditorMenus/Public/EditorMenu.h b/Engine/Source/Editor/EditorMenus/Public/EditorMenu.h
new file mode 100644
--- /dev/null
+++ b/Engine/Source/Editor/EditorMenus/Public/EditorMenu.h
@@ -0,0 +1,121 @@
+// Copyright 1998-2019 Epic Games, Inc. All Rights Reserved.
+
+#pragma once
+
+#include "CoreMinimal.h"
+
+#include "EditorMenuOwner.h"
+#include "EditorMenuDelegates.h"
+#include "EditorMenuSection.h"
+#include "EditorMenuContext.h"
+
+#include "Factories/Factory.h"
+
+#include "EditorMenu.generated.h"
+
+UCLASS(BlueprintType)
+class EDITORMENUS_API UEditorMenu : public UObject
+{
+	GENERATED_BODY()
+
+public:
+
+	UEditorMenu();
+
+	UFUNCTION(BlueprintCallable, Category = "Editor UI")
+	void InitMenu(const FEditorMenuOwner Owner, FName Name, FName Parent = NAME_None, EMultiBoxType Type = EMultiBoxType::Menu);
+
+	UFUNCTION(BlueprintCallable, Category = "Editor UI", meta = ( ScriptName = "AddSection", DisplayName = "Add Section", AutoCreateRefTerm = "Label", AdvancedDisplay = "InsertName,InsertType" ))
+	void AddSectionScript(const FName SectionName, const FText& Label = FText(), const FName InsertName = NAME_None, const EEditorMenuInsertType InsertType = EEditorMenuInsertType::Default);
+
+	UFUNCTION(BlueprintCallable, Category = "Editor UI", meta = ( ScriptName = "AddDynamicSection", DisplayName = "Add Dynamic Section" ))
+	void AddDynamicSectionScript(const FName SectionName, UEditorMenuSectionDynamic* Object);
+
+	UFUNCTION(BlueprintCallable, Category = "Editor UI")
+	void AddMenuEntry(const FName SectionName, const FEditorMenuEntry& Args);
+
+	UFUNCTION(BlueprintCallable, Category = "Editor UI")
+	void AddMenuEntryObject(UEditorMenuEntryScript* InObject);
+
+	UFUNCTION(BlueprintCallable, Category = "Editor UI", meta = ( AutoCreateRefTerm = "Label,ToolTip" ))
+	UEditorMenu* AddSubMenu(const FEditorMenuOwner Owner, const FName SectionName, const FName Name, const FText& Label, const FText& ToolTip = FText());
+
+	void RemoveSection(const FName SectionName);
+
+	FEditorMenuSection& AddSection(const FName SectionName, const TAttribute< FText >& InLabel = TAttribute<FText>(), const FEditorMenuInsert InPosition = FEditorMenuInsert());
+
+	FEditorMenuSection& AddDynamicSection(const FName SectionName, const FNewSectionConstructChoice& InConstruct, const FEditorMenuInsert InPosition = FEditorMenuInsert());
+
+	FEditorMenuSection* FindSection(const FName SectionName);
+
+	FEditorMenuSection& FindOrAddSection(const FName SectionName);
+
+	FName GetMenuName() const { return MenuName; }
+
+	bool IsRegistered() const { return bRegistered; }
+
+	template <typename TContextType>
+	TContextType* FindContext() const
+	{
+		return Context.Find<TContextType>();
+	}
+
+	friend class UEditorMenuSubsystem;
+
+private:
+
+	void InitGeneratedCopy(const UEditorMenu* Source);
+
+	bool FindEntry(const FName EntryName, int32& SectionIndex, int32& EntryIndex) const;
+
+	int32 IndexOfSection(const FName SectionName) const;
+
+	int32 FindInsertIndex(const FEditorMenuSection& InSection) const;
+
+public:
+
+	UPROPERTY(EditAnywhere, BlueprintReadWrite, Category = "Editor UI")
+	FName MenuName;
+
+	UPROPERTY(EditAnywhere, BlueprintReadWrite, Category = "Editor UI")
+	FName MenuParent;
+
+	UPROPERTY(EditAnywhere, BlueprintReadWrite, Category = "Editor UI")
+	FName StyleName;
+
+	UPROPERTY(EditAnywhere, BlueprintReadWrite, Category = "Editor UI")
+	FName TutorialHighlightName;
+
+	UPROPERTY(EditAnywhere, BlueprintReadWrite, Category = "Editor UI")
+	EMultiBoxType MenuType;
+
+	UPROPERTY(EditAnywhere, BlueprintReadWrite, Category = "Editor UI")
+	bool bShouldCloseWindowAfterMenuSelection;
+
+	UPROPERTY(EditAnywhere, BlueprintReadWrite, Category = "Editor UI")
+	bool bCloseSelfOnly;
+
+	UPROPERTY(EditAnywhere, BlueprintReadWrite, Category = "Editor UI")
+	bool bSearchable;
+
+	UPROPERTY(EditAnywhere, BlueprintReadWrite, Category = "ToolBar")
+	bool bToolBarIsFocusable;
+
+	UPROPERTY(EditAnywhere, BlueprintReadWrite, Category = "ToolBar")
+	bool bToolBarForceSmallIcons;
+
+	UPROPERTY(EditAnywhere, BlueprintReadWrite, Category = "Editor UI")
+	FEditorMenuOwner MenuOwner;
+
+private:
+
+	UPROPERTY()
+	FEditorMenuContext Context;
+
+	UPROPERTY()
+	TArray<FEditorMenuSection> Sections;
+
+	bool bRegistered;
+
+	const ISlateStyle* StyleSet;
+};
diff --git a/Engine/Source/Editor/EditorMenus/Public/EditorMenuContext.h b/Engine/Source/Editor/EditorMenus/Public/EditorMenuContext.h
new file mode 100644
--- /dev/null
+++ b/Engine/Source/Editor/EditorMenus/Public/EditorMenuContext.h
@@ -0,0 +1,78 @@
+// Copyright 1998-2019 Epic Games, Inc. All Rights Reserved.
+
+#pragma once
+
+#include "CoreMinimal.h"
+
+#include "Kismet/BlueprintFunctionLibrary.h"
+
+#include "EditorMenuContext.generated.h"
+
+class FUICommandList;
+class FTabManager;
+class FExtender;
+class FExtensibilityManager;
+
+UCLASS(BlueprintType, Abstract)
+class EDITORMENUS_API UEditorMenuContextBase : public UObject
+{
+	GENERATED_BODY()
+};
+
+UCLASS()
+class EDITORMENUS_API USlateTabManagerContext : public UEditorMenuContextBase
+{
+	GENERATED_BODY()
+public:
+
+	TWeakPtr<FTabManager> TabManager;
+};
+
+USTRUCT(BlueprintType)
+struct EDITORMENUS_API FEditorMenuContext
+{
+	GENERATED_BODY()
+public:
+	
+	FEditorMenuContext();
+	FEditorMenuContext(UObject* InContext);
+	FEditorMenuContext(TSharedPtr<FUICommandList> InCommandList, TSharedPtr<FExtender> InExtender = TSharedPtr<FExtender>(), UObject* InContext = nullptr);
+
+	template <typename TContextType>
+	TContextType* Find() const
+	{
+		for (const TWeakObjectPtr<UObject>& WeakObject : ContextObjects)
+		{
+			if (TContextType* Result = Cast<TContextType>(WeakObject.Get()))
+			{
+				return Result;
+			}
+		}
+
+		return nullptr;
+	}
+	
+	UObject* FindByClass(UClass* InClass) const;
+
+	void AppendCommandList(TSharedPtr<FUICommandList> InCommandList);
+	void AddExtender(const TSharedPtr<FExtender>& InExtender);
+	TSharedPtr<FExtender> GetAllExtenders();
+
+	void AppendObjects(const TArray<UObject*>& InObjects);
+	void AddObject(UObject* InObject);
+
+	friend class UEditorMenuSubsystem;
+	
+
+private:
+	
+	UPROPERTY()
+	TArray<UObject*> ContextObjects;
+
+	TSharedPtr<FUICommandList> CommandList;
+
+	// Prevent CommandLists from being deleted
+	TArray<TSharedPtr<const FUICommandList>> ReferencedCommandLists;
+	TSharedPtr<FExtensibilityManager> ExtensibilityManager;
+};
+
diff --git a/Engine/Source/Editor/EditorMenus/Public/EditorMenuDelegates.h b/Engine/Source/Editor/EditorMenus/Public/EditorMenuDelegates.h
new file mode 100644
--- /dev/null
+++ b/Engine/Source/Editor/EditorMenus/Public/EditorMenuDelegates.h
@@ -0,0 +1,121 @@
+// Copyright 1998-2019 Epic Games, Inc. All Rights Reserved.
+
+#pragma once
+
+#include "CoreMinimal.h"
+
+#include "EditorMenuContext.h"
+#include "Framework/Commands/UIAction.h"
+#include "Framework/Commands/UICommandList.h"
+#include "Framework/MultiBox/MultiBoxBuilder.h"
+
+#include "EditorMenuDelegates.generated.h"
+
+class UEditorMenu;
+struct FEditorMenuSection;
+
+
+
+DECLARE_DELEGATE_OneParam(FNewEditorMenuSectionDelegate, FEditorMenuSection&);
+DECLARE_DELEGATE_OneParam(FNewEditorMenuDelegate, UEditorMenu*);
+DECLARE_DELEGATE_TwoParams(FNewEditorMenuDelegateLegacy, class FMenuBuilder&, UEditorMenu*);
+DECLARE_DELEGATE_TwoParams(FNewToolBarDelegateLegacy, class FToolBarBuilder&, UEditorMenu*);
+
+DECLARE_DELEGATE_OneParam(FEditorMenuExecuteAction, const FEditorMenuContext&);
+DECLARE_DELEGATE_RetVal_OneParam(bool, FEditorMenuCanExecuteAction, const FEditorMenuContext&);
+DECLARE_DELEGATE_RetVal_OneParam(bool, FEditorMenuIsActionChecked, const FEditorMenuContext&);
+DECLARE_DELEGATE_RetVal_OneParam(ECheckBoxState, FEditorMenuGetActionCheckState, const FEditorMenuContext&);
+DECLARE_DELEGATE_RetVal_OneParam(bool, FEditorMenuIsActionButtonVisible, const FEditorMenuContext&);
+
+DECLARE_DELEGATE_TwoParams(FEditorMenuExecuteString, const FString&, const FEditorMenuContext&);
+
+DECLARE_DYNAMIC_DELEGATE_OneParam(FEditorMenuDynamicExecuteAction, const FEditorMenuContext&, Context);
+DECLARE_DYNAMIC_DELEGATE_RetVal_OneParam(bool, FEditorMenuDynamicCanExecuteAction, const FEditorMenuContext&, Context);
+DECLARE_DYNAMIC_DELEGATE_RetVal_OneParam(bool, FEditorMenuDynamicIsActionChecked, const FEditorMenuContext&, Context);
+DECLARE_DYNAMIC_DELEGATE_RetVal_OneParam(ECheckBoxState, FEditorMenuDynamicGetActionCheckState, const FEditorMenuContext&, Context);
+DECLARE_DYNAMIC_DELEGATE_RetVal_OneParam(bool, FEditorMenuDynamicIsActionButtonVisible, const FEditorMenuContext&, Context);
+
+struct EDITORMENUS_API FEditorUIAction
+{
+	FEditorUIAction() {}
+	FEditorUIAction(const FEditorMenuExecuteAction& InExecuteAction) : ExecuteAction(InExecuteAction) {}
+
+	FEditorMenuExecuteAction ExecuteAction;
+	FEditorMenuCanExecuteAction CanExecuteAction;
+	FEditorMenuGetActionCheckState GetActionCheckState;
+	FEditorMenuIsActionButtonVisible IsActionVisibleDelegate;
+};
+
+USTRUCT(BlueprintType)
+struct EDITORMENUS_API FEditorDynamicUIAction
+{
+	GENERATED_USTRUCT_BODY()
+
+	FEditorDynamicUIAction() {}
+
+	UPROPERTY(EditAnywhere, BlueprintReadWrite, Category = "Editor UI")
+	FEditorMenuDynamicExecuteAction ExecuteAction;
+
+	UPROPERTY(EditAnywhere, BlueprintReadWrite, Category = "Editor UI")
+	FEditorMenuDynamicCanExecuteAction CanExecuteAction;
+
+	UPROPERTY(EditAnywhere, BlueprintReadWrite, Category = "Editor UI")
+	FEditorMenuDynamicGetActionCheckState GetActionCheckState;
+
+	UPROPERTY(EditAnywhere, BlueprintReadWrite, Category = "Editor UI")
+	FEditorMenuDynamicIsActionButtonVisible IsActionVisibleDelegate;
+};
+
+struct EDITORMENUS_API FEditorUIActionChoice
+{
+public:
+	FEditorUIActionChoice() {}
+	FEditorUIActionChoice(const FUIAction& InAction) : Action(InAction) { }
+	FEditorUIActionChoice(const FExecuteAction& InExecuteAction) : Action(InExecuteAction) { }
+	FEditorUIActionChoice(const FEditorUIAction& InAction) : EditorAction(InAction) { }
+	FEditorUIActionChoice(const FEditorDynamicUIAction& InAction) : DynamicEditorAction(InAction) { }
+	FEditorUIActionChoice(const FEditorMenuExecuteAction& InExecuteAction) : EditorAction(InExecuteAction) {}
+	FEditorUIActionChoice(const TSharedPtr< const FUICommandInfo >& InCommand, const FUICommandList& InCommandList);
+
+	FUIAction* GetUIAction()
+	{
+		return Action.IsSet() ? &Action.GetValue() : nullptr;
+	}
+
+	FEditorUIAction* GetEditorUIAction()
+	{
+		return EditorAction.IsSet() ? &EditorAction.GetValue() : nullptr;
+	}
+
+	FEditorDynamicUIAction* GetEditorDynamicUIAction()
+	{
+		return DynamicEditorAction.IsSet() ? &DynamicEditorAction.GetValue() : nullptr;
+	}
+
+private:
+	TOptional<FUIAction> Action; // 72 bytes
+	TOptional<FEditorUIAction> EditorAction; // 64 bytes
+	TOptional<FEditorDynamicUIAction> DynamicEditorAction; // 80 bytes
+};
+
+struct EDITORMENUS_API FNewEditorMenuChoice
+{
+	FNewEditorMenuChoice() {};
+	FNewEditorMenuChoice(const FNewEditorMenuDelegate& InNewEditorMenuDelegate) : NewEditorMenuDelegate(InNewEditorMenuDelegate) {}
+	FNewEditorMenuChoice(const FNewMenuDelegate& InNewMenuDelegate) : NewMenuDelegate(InNewMenuDelegate) {}
+
+	FNewEditorMenuDelegate NewEditorMenuDelegate;
+	FNewMenuDelegate NewMenuDelegate;
+};
+
+struct EDITORMENUS_API FNewSectionConstructChoice
+{
+	FNewSectionConstructChoice() {};
+	FNewSectionConstructChoice(const FNewEditorMenuDelegate& InNewEditorMenuDelegate) : NewEditorMenuDelegate(InNewEditorMenuDelegate) {}
+	FNewSectionConstructChoice(const FNewEditorMenuDelegateLegacy& InNewEditorMenuDelegateLegacy) : NewEditorMenuDelegateLegacy(InNewEditorMenuDelegateLegacy) {}
+	FNewSectionConstructChoice(const FNewToolBarDelegateLegacy& InNewToolBarDelegateLegacy) : NewToolBarDelegateLegacy(InNewToolBarDelegateLegacy) {}
+
+	FNewEditorMenuDelegate NewEditorMenuDelegate;
+	FNewEditorMenuDelegateLegacy NewEditorMenuDelegateLegacy;
+	FNewToolBarDelegateLegacy NewToolBarDelegateLegacy;
+};
diff --git a/Engine/Source/Editor/EditorMenus/Public/EditorMenuEntry.h b/Engine/Source/Editor/EditorMenus/Public/EditorMenuEntry.h
new file mode 100644
--- /dev/null
+++ b/Engine/Source/Editor/EditorMenus/Public/EditorMenuEntry.h
@@ -0,0 +1,152 @@
+// Copyright 1998-2019 Epic Games, Inc. All Rights Reserved.
+
+#pragma once
+
+#include "CoreMinimal.h"
+
+#include "EditorMenuOwner.h"
+#include "EditorMenuDelegates.h"
+#include "EditorMenuMisc.h"
+
+#include "Misc/Attribute.h"
+#include "Framework/MultiBox/MultiBoxDefs.h"
+#include "Framework/Commands/UICommandInfo.h"
+#include "UObject/TextProperty.h"
+
+#include "EditorMenuEntry.generated.h"
+
+USTRUCT(BlueprintType)
+struct EDITORMENUS_API FSubMenuEntryData
+{
+	GENERATED_USTRUCT_BODY()
+
+	FSubMenuEntryData() :
+		bIsSubMenu(false),
+		bOpenSubMenuOnClick(false)
+	{
+	}
+
+	FSubMenuEntryData(const FNewEditorMenuChoice& InConstruct, bool bInOpenSubMenuOnClick = false) :
+		bIsSubMenu(true),
+		bOpenSubMenuOnClick(bInOpenSubMenuOnClick),
+		Construct(InConstruct)
+	{
+	}
+
+	UPROPERTY(EditAnywhere, BlueprintReadWrite, Category = "Editor UI")
+	bool bIsSubMenu;
+
+	UPROPERTY(EditAnywhere, BlueprintReadWrite, Category = "Editor UI")
+	bool bOpenSubMenuOnClick;
+
+	FNewEditorMenuChoice Construct;
+};
+
+USTRUCT(BlueprintType)
+struct EDITORMENUS_API FToolBarData
+{
+	GENERATED_USTRUCT_BODY()
+
+	FToolBarData() : bSimpleComboBox(false) {}
+
+	UPROPERTY(EditAnywhere, BlueprintReadWrite, Category = "Editor UI")
+	bool bSimpleComboBox;
+
+	FOnGetContent ComboButtonContextMenuGenerator;
+	FNewToolBarDelegateLegacy ConstructLegacy;
+};
+
+USTRUCT(BlueprintType)
+struct EDITORMENUS_API FEditorMenuStringCommand
+{
+	GENERATED_USTRUCT_BODY()
+
+	UPROPERTY(EditAnywhere, BlueprintReadWrite, Category = "Editor UI")
+	FName TypeName;
+
+	UPROPERTY(EditAnywhere, BlueprintReadWrite, Category = "Editor UI")
+	FString String;
+};
+
+USTRUCT(BlueprintType)
+struct EDITORMENUS_API FEditorMenuEntry
+{
+	GENERATED_USTRUCT_BODY()
+
+	FEditorMenuEntry();
+	FEditorMenuEntry(const FEditorMenuOwner InOwner, const FName InName, EMultiBlockType InType);
+
+	static FEditorMenuEntry InitMenuEntry(const FName InName, const TAttribute<FText>& InLabel, const TAttribute<FText>& InToolTip, const TAttribute<FSlateIcon>& InIcon, const FEditorUIActionChoice& InAction, const EUserInterfaceActionType UserInterfaceActionType = EUserInterfaceActionType::Button, const FName InTutorialHighlightName = NAME_None);
+	static FEditorMenuEntry InitMenuEntry(const TSharedPtr< const FUICommandInfo >& InCommand, const TAttribute<FText>& InLabelOverride = TAttribute<FText>(), const TAttribute<FText>& InToolTipOverride = TAttribute<FText>(), const TAttribute<FSlateIcon>& InIconOverride = TAttribute<FSlateIcon>(), const FName InTutorialHighlightName = NAME_None, const FName InNameOverride = NAME_None);
+	static FEditorMenuEntry InitSubMenu(const FName InParentMenu, const FName InName, const TAttribute<FText>& InLabel, const TAttribute<FText>& InToolTip, const FNewEditorMenuChoice& InMakeMenu, bool bInOpenSubMenuOnClick = false, const TAttribute<FSlateIcon>& InIcon = TAttribute<FSlateIcon>(), const bool ShouldCloseWindowAfterMenuSelection = true);
+
+	static FEditorMenuEntry InitToolBarButton(const FName InName, const FEditorUIActionChoice& InAction, const TAttribute<FText>& InLabel = TAttribute<FText>(), const TAttribute<FText>& InToolTip = TAttribute<FText>(), const TAttribute<FSlateIcon>& InIcon = TAttribute<FSlateIcon>(), const EUserInterfaceActionType UserInterfaceActionType = EUserInterfaceActionType::Button, FName InTutorialHighlightName = NAME_None);
+	static FEditorMenuEntry InitToolBarButton(const TSharedPtr< const FUICommandInfo >& InCommand, const TAttribute<FText>& InLabelOverride = TAttribute<FText>(), const TAttribute<FText>& InToolTipOverride = TAttribute<FText>(), const TAttribute<FSlateIcon>& InIconOverride = TAttribute<FSlateIcon>(), FName InTutorialHighlightName = NAME_None, const FName InNameOverride = NAME_None);
+	static FEditorMenuEntry InitComboButton(const FName InName, const FEditorUIActionChoice& InAction, const FOnGetContent& InMenuContentGenerator, const TAttribute<FText>& InLabelOverride = TAttribute<FText>(), const TAttribute<FText>& InToolTipOverride = TAttribute<FText>(), const TAttribute<FSlateIcon>& InIconOverride = TAttribute<FSlateIcon>(), bool bInSimpleComboBox = false, FName InTutorialHighlightName = NAME_None);
+
+	static FEditorMenuEntry InitMenuSeparator(const FName InName);
+	static FEditorMenuEntry InitToolBarSeparator(const FName InName);
+
+	friend struct FEditorMenuSection;
+	friend class UEditorMenuEntryScript;
+
+private:
+
+	void SetCommand(const TSharedPtr< const FUICommandInfo >& InCommand, FName InName, const TAttribute<FText>& InLabel, const TAttribute<FText>& InToolTip, const TAttribute<FSlateIcon>& InIcon);
+
+	void ResetActions();
+
+	bool IsScriptObjectDynamicConstruct() const;
+
+public:
+
+	UPROPERTY(EditAnywhere, BlueprintReadWrite, Category = "Editor UI")
+	FName Name;
+
+	UPROPERTY(EditAnywhere, BlueprintReadWrite, Category = "Editor UI")
+	FEditorMenuOwner Owner;
+
+	UPROPERTY(EditAnywhere, BlueprintReadWrite, Category = "Editor UI")
+	EMultiBlockType Type;
+
+	UPROPERTY(EditAnywhere, BlueprintReadWrite, Category = "Editor UI")
+	EUserInterfaceActionType UserInterfaceActionType;
+
+	UPROPERTY(EditAnywhere, BlueprintReadWrite, Category = "Editor UI")
+	FName TutorialHighlightName;
+
+	UPROPERTY(EditAnywhere, BlueprintReadWrite, Category = "Editor UI")
+	FEditorMenuInsert InsertPosition;
+
+	UPROPERTY(EditAnywhere, BlueprintReadWrite, Category = "Editor UI")
+	bool bShouldCloseWindowAfterMenuSelection;
+
+	UPROPERTY(EditAnywhere, BlueprintReadWrite, Category = "Editor UI")
+	FSubMenuEntryData SubMenu;
+
+	UPROPERTY(EditAnywhere, BlueprintReadWrite, Category = "Editor UI")
+	FToolBarData ToolBar;
+
+	UPROPERTY(EditAnywhere, BlueprintReadWrite, Category = "Editor UI")
+	UEditorMenuEntryScript* ScriptObject;
+
+private:
+
+	friend class UEditorMenuSubsystem;
+	friend class UEditorMenuEntryExtensions;
+
+	TAttribute<FText> Label;
+	TAttribute<FText> ToolTip;
+	TAttribute<FSlateIcon> Icon;
+
+	FEditorUIActionChoice Action;
+
+	FEditorMenuStringCommand StringCommand;
+
+	TSharedPtr< const FUICommandInfo > Command;
+	TSharedPtr< const FUICommandList > CommandList;
+
+	FNewEditorMenuSectionDelegate Construct;
+	FNewEditorMenuDelegateLegacy ConstructLegacy;
+};
+
diff --git a/Engine/Source/Editor/EditorMenus/Public/EditorMenuEntryScript.h b/Engine/Source/Editor/EditorMenus/Public/EditorMenuEntryScript.h
new file mode 100644
--- /dev/null
+++ b/Engine/Source/Editor/EditorMenus/Public/EditorMenuEntryScript.h
@@ -0,0 +1,174 @@
+// Copyright 1998-2019 Epic Games, Inc. All Rights Reserved.
+
+#pragma once
+
+#include "CoreMinimal.h"
+
+#include "UObject/ObjectMacros.h"
+#include "UObject/UObjectGlobals.h"
+#include "UObject/TextProperty.h"
+#include "Textures/SlateIcon.h"
+#include "Framework/MultiBox/MultiBoxDefs.h"
+#include "Styling/SlateTypes.h"
+
+#include "EditorMenuMisc.h"
+#include "EditorMenuContext.h"
+#include "EditorMenuSection.h"
+
+#include "EditorMenuEntryScript.generated.h"
+
+struct FEditorMenuEntry;
+struct FEditorMenuSection;
+
+USTRUCT(BlueprintType)
+struct EDITORMENUS_API FScriptSlateIcon
+{
+	GENERATED_BODY()
+
+public:
+	FScriptSlateIcon();
+	FScriptSlateIcon(const FName InStyleSetName, const FName InStyleName);
+	FScriptSlateIcon(const FName InStyleSetName, const FName InStyleName, const FName InSmallStyleName);
+
+	UPROPERTY(EditAnywhere, BlueprintReadWrite, Category = "Editor UI")
+	FName StyleSetName;
+
+	UPROPERTY(EditAnywhere, BlueprintReadWrite, Category = "Editor UI")
+	FName StyleName;
+
+	UPROPERTY(EditAnywhere, BlueprintReadWrite, Category = "Editor UI")
+	FName SmallStyleName;
+
+	operator FSlateIcon() const { return GetSlateIcon(); }
+
+	FSlateIcon GetSlateIcon() const;
+};
+
+USTRUCT(BlueprintType)
+struct EDITORMENUS_API FEditorMenuEntryScriptDataAdvanced
+{
+	GENERATED_BODY()
+
+public:
+
+	FEditorMenuEntryScriptDataAdvanced();
+
+	UPROPERTY(EditAnywhere, BlueprintReadWrite, Category = "Advanced")
+	FName TutorialHighlight;
+
+	UPROPERTY(EditAnywhere, BlueprintReadWrite, Category = "Advanced")
+	EMultiBlockType EntryType;
+
+	UPROPERTY(EditAnywhere, BlueprintReadWrite, Category = "Advanced")
+	EUserInterfaceActionType UserInterfaceActionType;
+
+	UPROPERTY(EditAnywhere, BlueprintReadWrite, Category = "SubMenu")
+	bool bIsSubMenu;
+
+	UPROPERTY(EditAnywhere, BlueprintReadWrite, Category = "SubMenu")
+	bool bOpenSubMenuOnClick;
+
+	UPROPERTY(EditAnywhere, BlueprintReadWrite, Category = "Advanced")
+	bool bShouldCloseWindowAfterMenuSelection;
+
+	UPROPERTY(EditAnywhere, BlueprintReadWrite, Category = "ToolBar")
+	bool bSimpleComboBox;
+};
+
+USTRUCT(BlueprintType)
+struct EDITORMENUS_API FEditorMenuEntryScriptData
+{
+	GENERATED_BODY()
+
+public:
+
+	UPROPERTY(EditAnywhere, BlueprintReadWrite, Category = "Identity")
+	FName Menu;
+
+	UPROPERTY(EditAnywhere, BlueprintReadWrite, Category = "Identity")
+	FName Section;
+
+	UPROPERTY(EditAnywhere, BlueprintReadWrite, Category = "Identity")
+	FName Name;
+
+	UPROPERTY(EditAnywhere, BlueprintReadWrite, Category = "Appearance")
+	FText Label;
+
+	UPROPERTY(EditAnywhere, BlueprintReadWrite, Category = "Appearance")
+	FText ToolTip;
+
+	UPROPERTY(EditAnywhere,  BlueprintReadWrite, Category = "Appearance")
+	FScriptSlateIcon Icon;
+
+	// Optional identifier used for unregistering a group of menu items
+	UPROPERTY(EditAnywhere, BlueprintReadWrite, Category = "Advanced")
+	FName OwnerName;
+
+	UPROPERTY(EditAnywhere, BlueprintReadWrite, Category = "Identity")
+	FEditorMenuInsert InsertPosition;
+
+	UPROPERTY(EditAnywhere,  BlueprintReadWrite, Category = "Advanced")
+	FEditorMenuEntryScriptDataAdvanced Advanced;
+};
+
+UCLASS(Blueprintable, abstract)
+class EDITORMENUS_API UEditorMenuEntryScript : public UObject
+{
+	GENERATED_BODY()
+
+public:
+
+	UFUNCTION(BlueprintImplementableEvent, Category = "Action")
+	void Execute(const FEditorMenuContext& Context);
+
+	UFUNCTION(BlueprintNativeEvent, Category = "Advanced")
+	bool CanExecute(const FEditorMenuContext& Context) const;
+	virtual bool CanExecute_Implementation(const FEditorMenuContext& Context) const { return true; }
+
+	UFUNCTION(BlueprintNativeEvent, Category = "Advanced")
+	ECheckBoxState GetCheckState(const FEditorMenuContext& Context) const;
+	virtual ECheckBoxState GetCheckState_Implementation(const FEditorMenuContext& Context) const { return ECheckBoxState::Undetermined; }
+
+	UFUNCTION(BlueprintNativeEvent, Category = "Advanced")
+	bool IsVisible(const FEditorMenuContext& Context) const;
+	virtual bool IsVisible_Implementation(const FEditorMenuContext& Context) const { return true; }
+
+	UFUNCTION(BlueprintNativeEvent, Category = "Advanced")
+	FText GetLabel(const FEditorMenuContext& Context) const;
+	virtual FText GetLabel_Implementation(const FEditorMenuContext& Context) const { return Data.Label; }
+
+	UFUNCTION(BlueprintNativeEvent, Category = "Advanced")
+	FText GetToolTip(const FEditorMenuContext& Context) const;
+	virtual FText GetToolTip_Implementation(const FEditorMenuContext& Context) const { return Data.ToolTip; }
+
+	UFUNCTION(BlueprintNativeEvent, Category = "Advanced")
+	FScriptSlateIcon GetIcon(const FEditorMenuContext& Context) const;
+	virtual FScriptSlateIcon GetIcon_Implementation(const FEditorMenuContext& Context) const { return Data.Icon; }
+
+	UFUNCTION(BlueprintImplementableEvent, Category = "Advanced")
+	void ConstructMenuEntry(FEditorMenuSection& Section, const FEditorMenuContext& Context);
+
+	UFUNCTION(BlueprintCallable, Category = "Advanced")
+	void RegisterMenuEntry();
+
+	UFUNCTION(BlueprintCallable, Category = "Advanced")
+	void InitEntry(const FName OwnerName, const FName Menu, const FName Section, const FName Name, const FText& Label = FText(), const FText& ToolTip = FText());
+
+private:
+
+	friend struct FEditorMenuSection;
+	friend class UEditorMenuSubsystem;
+
+	TAttribute<FText> CreateLabelAttribute(FEditorMenuContext& Context);
+
+	TAttribute<FText> CreateToolTipAttribute(FEditorMenuContext& Context);
+
+	void ToMenuEntry(FEditorMenuEntry& Output);
+
+	bool IsDynamicConstruct() const;
+
+public:
+
+	UPROPERTY(EditAnywhere, BlueprintReadWrite, Category = "Data")
+	FEditorMenuEntryScriptData Data;
+};
diff --git a/Engine/Source/Editor/EditorMenus/Public/EditorMenuMisc.h b/Engine/Source/Editor/EditorMenus/Public/EditorMenuMisc.h
new file mode 100644
--- /dev/null
+++ b/Engine/Source/Editor/EditorMenus/Public/EditorMenuMisc.h
@@ -0,0 +1,51 @@
+// Copyright 1998-2019 Epic Games, Inc. All Rights Reserved.
+
+#pragma once
+
+#include "CoreMinimal.h"
+
+#include "EditorMenuMisc.generated.h"
+
+UENUM(BlueprintType)
+enum class EEditorMenuInsertType : uint8
+{
+	Default,
+	Before,
+	After,
+	First
+};
+
+USTRUCT(BlueprintType)
+struct FEditorMenuInsert
+{
+	GENERATED_USTRUCT_BODY()
+
+	FEditorMenuInsert() : Position(EEditorMenuInsertType::Default) {}
+	FEditorMenuInsert(FName InName, EEditorMenuInsertType InPosition) : Name(InName), Position(InPosition) {}
+
+	UPROPERTY(EditAnywhere, BlueprintReadWrite, Category = "Editor UI")
+	FName Name;
+
+	UPROPERTY(EditAnywhere, BlueprintReadWrite, Category = "Editor UI")
+	EEditorMenuInsertType Position;
+
+	FORCEINLINE bool operator==(const FEditorMenuInsert& Other) const
+	{
+		return Other.Name == Name && Other.Position == Position;
+	}
+
+	FORCEINLINE bool operator!=(const FEditorMenuInsert& Other) const
+	{
+		return Other.Name != Name || Other.Position != Position;
+	}
+
+	bool IsDefault() const
+	{
+		return Position == EEditorMenuInsertType::Default;
+	}
+
+	bool IsBeforeOrAfter() const
+	{
+		return Position == EEditorMenuInsertType::Before || Position == EEditorMenuInsertType::After;
+	}
+};
diff --git a/Engine/Source/Editor/EditorMenus/Public/EditorMenuOwner.h b/Engine/Source/Editor/EditorMenus/Public/EditorMenuOwner.h
new file mode 100644
--- /dev/null
+++ b/Engine/Source/Editor/EditorMenus/Public/EditorMenuOwner.h
@@ -0,0 +1,88 @@
+// Copyright 1998-2019 Epic Games, Inc. All Rights Reserved.
+
+#pragma once
+
+#include "CoreMinimal.h"
+
+#include "EditorMenuOwner.generated.h"
+
+
+USTRUCT(BlueprintType, meta=(HasNativeBreak="EditorMenus.EditorMenuEntryExtensions.BreakEditorMenuOwner", HasNativeMake="EditorMenus.EditorMenuEntryExtensions.MakeEditorMenuOwner"))
+struct EDITORMENUS_API FEditorMenuOwner
+{
+	GENERATED_USTRUCT_BODY()
+
+private:
+	struct FStoreName
+	{
+		int32 Index;
+		int32 Number;
+	};
+
+	enum class EValueType : uint8
+	{
+		None,
+		Pointer,
+		Name
+	};
+
+public:
+	FORCEINLINE FEditorMenuOwner() : ValueInt64(0), ValueType(EValueType::None) {}
+	FORCEINLINE FEditorMenuOwner(void* InPointer) : ValueInt64(reinterpret_cast<int64>(InPointer)), ValueType(EValueType::Pointer) {}
+
+	FEditorMenuOwner(ANSICHAR* InValue) : FEditorMenuOwner(FName(InValue)) {}
+	FEditorMenuOwner(TCHAR* InValue) : FEditorMenuOwner(FName(InValue)) {}
+
+	FEditorMenuOwner(const FName InValue)
+	{
+		if (InValue == NAME_None)
+		{
+			ValueInt64 = 0;
+			ValueType = EValueType::None;
+		}
+		else
+		{
+			ValueName.Index = InValue.GetComparisonIndex().ToUnstableInt();
+			ValueName.Number = InValue.GetNumber();
+			ValueType = EValueType::Name;
+		}
+	}
+
+	FORCEINLINE bool operator==(const FEditorMenuOwner& Other) const
+	{
+		return Other.ValueInt64 == ValueInt64 && Other.ValueType == ValueType;
+	}
+
+	FORCEINLINE bool operator!=(const FEditorMenuOwner& Other) const
+	{
+		return Other.ValueInt64 != ValueInt64 || Other.ValueType != ValueType;
+	}
+
+	friend uint32 GetTypeHash(const FEditorMenuOwner& Key)
+	{
+		return GetTypeHash(Key.ValueInt64);
+	}
+
+	FORCEINLINE bool IsSet() const { return ValueInt64 != 0; }
+
+	FName TryGetName() const
+	{
+		if (ValueType == EValueType::Name)
+		{
+			const FNameEntryId EntryId = FNameEntryId::FromUnstableInt(ValueName.Index);
+			return FName(EntryId, EntryId, ValueName.Number);
+		}
+
+		return NAME_None;
+	}
+
+private:
+
+	union
+	{
+		int64 ValueInt64;
+		FStoreName ValueName;
+	};
+
+	EValueType ValueType;
+};
diff --git a/Engine/Source/Editor/EditorMenus/Public/EditorMenuSection.h b/Engine/Source/Editor/EditorMenus/Public/EditorMenuSection.h
new file mode 100644
--- /dev/null
+++ b/Engine/Source/Editor/EditorMenus/Public/EditorMenuSection.h
@@ -0,0 +1,83 @@
+// Copyright 1998-2019 Epic Games, Inc. All Rights Reserved.
+
+#pragma once
+
+#include "CoreMinimal.h"
+
+#include "EditorMenuDelegates.h"
+#include "EditorMenuEntry.h"
+#include "Misc/Attribute.h"
+
+#include "EditorMenuSection.generated.h"
+
+UCLASS(Blueprintable)
+class EDITORMENUS_API UEditorMenuSectionDynamic : public UObject
+{
+	GENERATED_BODY()
+
+public:
+
+	UFUNCTION(BlueprintImplementableEvent, Category = "Editor UI")
+	void ConstructSections(UEditorMenu* Menu, const FEditorMenuContext& Context);
+};
+
+USTRUCT(BlueprintType)
+struct EDITORMENUS_API FEditorMenuSection
+{
+	GENERATED_USTRUCT_BODY()
+
+public:
+
+	FEditorMenuSection();
+
+	void InitSection(const FName InName, const TAttribute< FText >& InLabel, const FEditorMenuInsert InPosition);
+
+	FEditorMenuEntry& AddEntryObject(UEditorMenuEntryScript* InObject);
+	FEditorMenuEntry& AddEntry(const FEditorMenuEntry& Args);
+	FEditorMenuEntry& AddMenuEntry(const TSharedPtr< const FUICommandInfo >& InCommand);
+
+	FEditorMenuEntry& AddDynamicEntry(const FName InName, const FNewEditorMenuSectionDelegate& InConstruct);
+	FEditorMenuEntry& AddDynamicEntry(const FName InName, const FNewEditorMenuDelegateLegacy& InConstruct);
+	FEditorMenuEntry& AddMenuSeparator(const FName InName);
+
+	template <typename TContextType>
+	TContextType* FindContext() const
+	{
+		return Context.Find<TContextType>();
+	}
+
+private:
+	void InitGeneratedSectionCopy(const FEditorMenuSection& Source, FEditorMenuContext& InContext);
+
+	int32 RemoveEntry(const FName InName);
+	int32 RemoveEntriesByOwner(const FEditorMenuOwner InOwner);
+
+	int32 IndexOfBlock(const FName InName) const;
+	int32 FindBlockInsertIndex(const FEditorMenuEntry& InBlock) const;
+
+	void AssembleBlock(const FEditorMenuEntry& InBlock);
+
+	friend class UEditorMenuSectionExtensions;
+	friend class UEditorMenuSubsystem;
+	friend class UEditorMenu;
+
+public:
+	UPROPERTY(EditAnywhere, BlueprintReadWrite, Category = "Editor UI")
+	FName Name;
+
+	UPROPERTY(EditAnywhere, BlueprintReadWrite, Category = "Editor UI")
+	TArray<FEditorMenuEntry> Blocks;
+
+	UPROPERTY(EditAnywhere, BlueprintReadWrite, Category = "Editor UI")
+	FEditorMenuInsert InsertPosition;
+
+	UPROPERTY(EditAnywhere, BlueprintReadWrite, Category = "Editor UI")
+	FEditorMenuContext Context;
+
+	UPROPERTY(EditAnywhere, BlueprintReadWrite, Category = "Editor UI")
+	UEditorMenuSectionDynamic* EditorMenuSectionDynamic;
+
+	TAttribute<FText> Label;
+
+	FNewSectionConstructChoice Construct;
+};
diff --git a/Engine/Source/Editor/EditorMenus/Public/EditorMenuSubsystem.h b/Engine/Source/Editor/EditorMenus/Public/EditorMenuSubsystem.h
new file mode 100644
--- /dev/null
+++ b/Engine/Source/Editor/EditorMenus/Public/EditorMenuSubsystem.h
@@ -0,0 +1,297 @@
+// Copyright 1998-2019 Epic Games, Inc. All Rights Reserved.
+
+#pragma once
+
+#include "CoreMinimal.h"
+#include "UObject/ObjectMacros.h"
+#include "UObject/UObjectGlobals.h"
+#include "EditorSubsystem.h"
+#include "Templates/SubclassOf.h"
+#include "Textures/SlateIcon.h"
+#include "Framework/Commands/UIAction.h"
+#include "Framework/MultiBox/MultiBox.h"
+#include "Framework/MultiBox/MultiBoxBuilder.h"
+#include "Framework/MultiBox/MultiBoxExtender.h"
+#include "Framework/MultiBox/MultiBoxDefs.h"
+#include "Misc/Attribute.h"
+#include "Editor.h"
+
+#include "IEditorMenusModule.h"
+#include "EditorMenu.h"
+#include "EditorMenuContext.h"
+#include "EditorMenuDelegates.h"
+#include "EditorMenuEntry.h"
+#include "EditorMenuEntryScript.h"
+#include "EditorMenuOwner.h"
+#include "EditorMenuSection.h"
+#include "EditorMenuMisc.h"
+
+#include "EditorMenuSubsystem.generated.h"
+
+class FMenuBuilder;
+class FUICommandInfo;
+class FMultiBox;
+class SWidget;
+
+class UEditorMenu;
+class UEditorMenuEntryScript;
+struct FEditorMenuEntry;
+struct FEditorMenuSection;
+//struct FEditorMenuOwnerScoped;
+
+USTRUCT()
+struct FCustomizedEditorMenuSection
+{
+	GENERATED_BODY()
+
+	UPROPERTY()
+	FName Name;
+
+	UPROPERTY()
+	TArray<FName> Items;
+};
+
+USTRUCT()
+struct FCustomizedEditorMenu
+{
+	GENERATED_BODY()
+
+	UPROPERTY()
+	FName Name;
+
+	UPROPERTY()
+	TArray<FCustomizedEditorMenuSection> Sections;
+
+	UPROPERTY()
+	TArray<FName> HiddenItems;
+
+	UPROPERTY()
+	TArray<FName> HiddenSections;
+};
+
+USTRUCT()
+struct FGeneratedEditorMenuWidget
+{
+	GENERATED_BODY()
+
+	UPROPERTY()
+	UEditorMenu* GeneratedMenu;
+
+	TWeakPtr<SWidget> Widget;
+};
+
+USTRUCT()
+struct FGeneratedEditorMenuWidgets
+{
+	GENERATED_BODY()
+
+	UPROPERTY()
+	TArray<FGeneratedEditorMenuWidget> Instances;
+};
+
+UCLASS(config = EditorLayout)
+class EDITORMENUS_API UEditorMenuSubsystem : public UEditorSubsystem
+{
+	GENERATED_BODY()
+
+public:
+
+	UEditorMenuSubsystem();
+
+	static inline UEditorMenuSubsystem* Get()
+	{
+		if (GEditor)
+		{
+			FModuleManager::LoadModuleChecked<IEditorMenusModule>("EditorMenus");
+			return GEditor->GetEditorSubsystem<UEditorMenuSubsystem>();
+		}
+
+		return nullptr;
+	}
+
+	static inline UEditorMenuSubsystem* GetIfLoaded()
+	{
+		if (GEditor && IEditorMenusModule::IsAvailable())
+		{
+			return GEditor->GetEditorSubsystem<UEditorMenuSubsystem>();
+		}
+
+		return nullptr;
+	}
+
+	// Unregister all entries with a specific owner
+	static inline void UnregisterOwner(FEditorMenuOwner Owner)
+	{
+		if (UEditorMenuSubsystem* EditorMenus = UEditorMenuSubsystem::GetIfLoaded())
+		{
+			EditorMenus->UnregisterOwnerInternal(Owner);
+		}
+	}
+
+	// UEditorSubsystem implementation
+	virtual void Initialize(FSubsystemCollectionBase& Collection);
+	virtual void Deinitialize();
+
+
+	UEditorMenu* GenerateMenu(const FName Name, FEditorMenuContext& InMenuContext);
+	UEditorMenu* GenerateMenu(const TArray<UEditorMenu*>& Hierarchy, FEditorMenuContext& InMenuContext);
+	UEditorMenu* GenerateMenuAsBuilder(const UEditorMenu* InMenu, FEditorMenuContext& InMenuContext);
+
+	// Generate switch based on registered name
+	TSharedRef<SWidget> GenerateWidget(const FName Name, FEditorMenuContext& InMenuContext);
+	TSharedRef<SWidget> GenerateWidget(const TArray<UEditorMenu*>& Hierarchy, FEditorMenuContext& InMenuContext);
+	TSharedRef<SWidget> GenerateWidget(UEditorMenu* GeneratedMenu);
+
+	UFUNCTION(BlueprintCallable, Category = "Editor UI")
+	UEditorMenu* RegisterMenu(FName Name, const FName Parent = NAME_None, EMultiBoxType Type = EMultiBoxType::Menu);
+
+	UFUNCTION(BlueprintCallable, Category = "Editor UI")
+	UEditorMenu* ExtendMenu(const FName Name);
+
+	UFUNCTION(BlueprintCallable, Category = "Editor UI")
+	UEditorMenu* FindMenu(const FName Name);
+
+	// Rebuilds all widgets generated from a specific menu
+	UFUNCTION(BlueprintCallable, Category = "Editor UI")
+	bool RefreshMenuWidget(const FName Name);
+	
+	// Rebuilds all generated widgets next tick
+	UFUNCTION(BlueprintCallable, Category = "Editor UI")
+	void RefreshAllWidgets();
+
+	UFUNCTION(BlueprintCallable, Category = "Editor UI")
+	static bool AddMenuEntryObject(UEditorMenuEntryScript* MenuEntryObject);
+
+	UFUNCTION(BlueprintCallable, Category = "Editor UI")
+	void UnregisterOwnerByName(FName InOwnerName);
+
+	UFUNCTION(BlueprintCallable, Category = "Editor UI")
+	void SetSectionLabel(const FName MenuName, const FName SectionName, const FText Label);
+
+	UFUNCTION(BlueprintCallable, Category = "Editor UI")
+	void SetSectionPosition(const FName MenuName, const FName SectionName, const FName PositionName, const EEditorMenuInsertType PositionType);
+
+	void AddSection(const FName MenuName, const FName SectionName, const TAttribute< FText >& InLabel, const FEditorMenuInsert InPosition = FEditorMenuInsert());
+
+	void AddEntry(const FName MenuName, const FName SectionName, const FEditorMenuEntry& Entry);
+
+	UFUNCTION(BlueprintCallable, Category = "Editor UI")
+	void RemoveEntry(const FName MenuName, const FName Section, const FName Name);
+
+	UFUNCTION(BlueprintCallable, Category = "Editor UI")
+	void RemoveSection(const FName MenuName, const FName Section);
+
+	UFUNCTION(BlueprintCallable, Category = "Editor UI")
+	void RemoveMenu(const FName MenuName);
+
+	UFUNCTION(BlueprintCallable, Category = "Editor UI")
+	static UObject* FindContext(const FEditorMenuContext& InContext, UClass* InClass);
+
+	void AssembleMenuByName(UEditorMenu* GeneratedMenu, const FName Name);
+	void AssembleMenuHierarchy(UEditorMenu* GeneratedMenu, const TArray<UEditorMenu*>& Hierarchy);
+
+	FEditorMenuOwner CurrentOwner() const;
+
+	void RegisterStringCommandHandler(const FName InName, const FEditorMenuExecuteString& InDelegate);
+	void UnregisterStringCommandHandler(const FName InName);
+
+	friend struct FEditorMenuOwnerScoped;
+
+	//~ Begin UObject Interface
+	static void AddReferencedObjects(UObject* InThis, FReferenceCollector& Collector);
+	//~ End UObject Interface
+
+private:
+
+	void HandleRefreshAllWidgetsNextTick();
+
+	bool RefreshMenuWidget(const FName Name, FGeneratedEditorMenuWidget& GeneratedMenuWidget);
+	
+	void CleanupStaleGeneratedMenus();
+
+	TArray<UEditorMenu*> CollectHierarchy(const FName Name);
+
+	void PushOwner(const FEditorMenuOwner InOwner);
+
+	void PopOwner(const FEditorMenuOwner InOwner);
+
+	UEditorMenu* FindSubMenuToGenerateWith(const FName InParentName, const FName InChildName);
+
+	void FillMenu(FMenuBuilder& MenuBuilder, FName InMenuName, FEditorMenuContext InMenuContext);
+	void FillMenuBarDropDown(FMenuBuilder& MenuBuilder, FName InParentName, FName InChildName, FEditorMenuContext InMenuContext);
+	void PopulateMenuBuilder(FMenuBuilder& MenuBuilder, UEditorMenu* MenuData);
+	void PopulateMenuBarBuilder(FMenuBarBuilder& MenuBarBuilder, UEditorMenu* MenuData);
+	void PopulateToolBarBuilder(FToolBarBuilder& ToolBarBuilder, UEditorMenu* MenuData);
+
+	TSharedRef<SWidget> GenerateToolbarComboButtonMenu(const FName SubMenuFullName, FEditorMenuContext InContext);
+
+	FUIAction ConvertUIAction(const FEditorMenuStringCommand& StringCommand, const FEditorMenuContext& Context) const;
+	static FUIAction ConvertUIAction(FEditorUIActionChoice& Actions, FEditorMenuContext& Context);
+	static FUIAction ConvertUIAction(FEditorUIAction& Actions, FEditorMenuContext& Context);
+	static FUIAction ConvertUIAction(FEditorDynamicUIAction& Actions, FEditorMenuContext& Context);
+	static FUIAction ConvertUIAction(UEditorMenuEntryScript* ScriptObject, FEditorMenuContext& Context);
+
+	void ExecuteStringCommand(const FEditorMenuStringCommand StringCommand, const FEditorMenuContext Context);
+
+	void FillMenuDynamic(FMenuBuilder& Builder, FNewEditorMenuDelegate InConstructMenu);
+
+	void ListAllParents(const FName Name, TArray<FName>& AllParents);
+
+	void AssembleMenu(UEditorMenu* GeneratedMenu, const UEditorMenu* Other);
+	void AssembleMenuSection(UEditorMenu* GeneratedMenu, const UEditorMenu* Other, FEditorMenuSection* DestSection, const FEditorMenuSection& OtherSection);
+
+	void CopyMenuSettings(UEditorMenu* GeneratedMenu, const UEditorMenu* Other);
+
+	void AddReferencedContextObjects(const TSharedRef<FMultiBox>& InMultiBox, const FEditorMenuContext& InMenuContext);
+
+	void ApplyCustomization(UEditorMenu* GeneratedMenu);
+
+	FCustomizedEditorMenu* FindCustomizedMenu(const FName InName);
+	int32 FindCustomizedMenuIndex(const FName InName);
+
+	void UnregisterOwnerInternal(FEditorMenuOwner Owner);
+
+	static FName JoinMenuPaths(const FName Base, const FName Child);
+
+	static bool GetDisplayUIExtensionPoints();
+
+public:
+
+	UPROPERTY(config, EditAnywhere, Category = Misc)
+	TArray<FCustomizedEditorMenu> CustomizedMenus;
+
+private:
+
+	UPROPERTY()
+	TMap<FName, UEditorMenu*> Menus;
+
+	UPROPERTY()
+	TMap<FName, FGeneratedEditorMenuWidgets> GeneratedMenuWidgets;
+
+	TMap<TWeakPtr<FMultiBox>, TArray<UObject*>> WidgetObjectReferences;
+
+	TArray<FEditorMenuOwner> OwnerStack;
+
+	TMap<FName, FEditorMenuExecuteString> StringCommandHandlers;
+
+	bool bRefreshWidgetsNextTick;
+};
+
+struct FEditorMenuOwnerScoped
+{
+	FEditorMenuOwnerScoped(const FEditorMenuOwner InOwner) : Owner(InOwner)
+	{
+		UEditorMenuSubsystem::Get()->PushOwner(InOwner);
+	}
+
+	~FEditorMenuOwnerScoped()
+	{
+		UEditorMenuSubsystem::Get()->PopOwner(Owner);
+	}
+
+	FEditorMenuOwner GetOwner() const { return Owner; }
+
+private:
+
+	FEditorMenuOwner Owner;
+};
diff --git a/Engine/Source/Editor/EditorMenus/Public/IEditorMenusModule.h b/Engine/Source/Editor/EditorMenus/Public/IEditorMenusModule.h
new file mode 100644
--- /dev/null
+++ b/Engine/Source/Editor/EditorMenus/Public/IEditorMenusModule.h
@@ -0,0 +1,31 @@
+// Copyright 1998-2019 Epic Games, Inc. All Rights Reserved.
+#pragma once
+
+#include "CoreMinimal.h"
+#include "Modules/ModuleInterface.h"
+#include "Modules/ModuleManager.h"
+
+DEFINE_LOG_CATEGORY_STATIC(LogEditorMenus, Log, All);
+
+class IEditorMenusModule : public IModuleInterface
+{
+public:
+
+	/**
+	 * Retrieve the module instance.
+	 */
+	static inline IEditorMenusModule& Get()
+	{
+		return FModuleManager::LoadModuleChecked<IEditorMenusModule>("EditorMenus");
+	}
+
+	/**
+	 * Checks to see if this module is loaded and ready.  It is only valid to call Get() if IsAvailable() returns true.
+	 *
+	 * @return True if the module is loaded and ready to use
+	 */
+	static inline bool IsAvailable()
+	{
+		return FModuleManager::Get().IsModuleLoaded("EditorMenus");
+	}
+};
diff --git a/Engine/Source/Editor/UnrealEd/UnrealEd.Build.cs b/Engine/Source/Editor/UnrealEd/UnrealEd.Build.cs
--- a/Engine/Source/Editor/UnrealEd/UnrealEd.Build.cs
+++ b/Engine/Source/Editor/UnrealEd/UnrealEd.Build.cs
@@ -168,6 +168,7 @@ public class UnrealEd : ModuleRules
 				"PakFileUtilities",
 				"TimeManagement",
                 "DerivedDataCache",
+				"EditorMenus",
             }
 		);
 
