Fixing compile warnings
#rnx

diff --git a/Engine/Source/Editor/EditorMenus/EditorMenus.Build.cs b/Engine/Source/Editor/EditorMenus/EditorMenus.Build.cs
--- a/Engine/Source/Editor/EditorMenus/EditorMenus.Build.cs
+++ b/Engine/Source/Editor/EditorMenus/EditorMenus.Build.cs
@@ -6,9 +6,6 @@ namespace UnrealBuildTool.Rules
 	{
 		public EditorMenus(ReadOnlyTargetRules Target) : base(Target)
 		{
-			//bFasterWithoutUnity = true;
-			//PCHUsage = PCHUsageMode.NoPCHs;
-
 			PrivateDependencyModuleNames.AddRange(
 				new string[] {
 					"Core",
@@ -16,15 +13,10 @@ namespace UnrealBuildTool.Rules
 					"Slate",
 					"SlateCore",
 					"EditorSubSystem",
-					"Engine", // For Subsystems
+					"Engine",
 					"UnrealEd",
 				}
 			);
-
-			//PrivateIncludePathModuleNames.AddRange(
-			//	new string[] {
-			//	"UnrealEd",
-			//	});
 		}
 	}
 }
diff --git a/Engine/Source/Editor/EditorMenus/Private/EditorMenuSection.cpp b/Engine/Source/Editor/EditorMenus/Private/EditorMenuSection.cpp
--- a/Engine/Source/Editor/EditorMenus/Private/EditorMenuSection.cpp
+++ b/Engine/Source/Editor/EditorMenus/Private/EditorMenuSection.cpp
@@ -63,14 +63,17 @@ FEditorMenuEntry& FEditorMenuSection::AddEntryObject(UEditorMenuEntryScript* InO
 		DestObject = DuplicateObject<UEditorMenuEntryScript>(InObject, UEditorMenuSubsystem::Get());
 	}
 
+	// Refresh widgets next tick so that toolbars and menu bars are updated
+	UEditorMenuSubsystem::Get()->RefreshAllWidgets();
+
 	FEditorMenuEntry Args;
 	DestObject->ToMenuEntry(Args);
 	return AddEntry(Args);
 }
 
-FEditorMenuEntry& FEditorMenuSection::AddMenuEntry(const TSharedPtr< const FUICommandInfo >& InCommand)
+FEditorMenuEntry& FEditorMenuSection::AddMenuEntry(const TSharedPtr< const FUICommandInfo >& InCommand, const TAttribute<FText>& InLabelOverride, const TAttribute<FText>& InToolTipOverride, const TAttribute<FSlateIcon>& InIconOverride, const FName InTutorialHighlightName, const FName InNameOverride)
 {
-	return AddEntry(FEditorMenuEntry::InitMenuEntry(InCommand));
+	return AddEntry(FEditorMenuEntry::InitMenuEntry(InCommand, InLabelOverride, InToolTipOverride, InIconOverride, InTutorialHighlightName, InNameOverride));
 }
 
 FEditorMenuEntry& FEditorMenuSection::AddDynamicEntry(const FName InName, const FNewEditorMenuSectionDelegate& InConstruct)
diff --git a/Engine/Source/Editor/EditorMenus/Private/EditorMenuSubsystem.cpp b/Engine/Source/Editor/EditorMenus/Private/EditorMenuSubsystem.cpp
--- a/Engine/Source/Editor/EditorMenus/Private/EditorMenuSubsystem.cpp
+++ b/Engine/Source/Editor/EditorMenus/Private/EditorMenuSubsystem.cpp
@@ -36,18 +36,18 @@ FEditorUIActionChoice::FEditorUIActionChoice(const TSharedPtr< const FUICommandI
 
 UEditorMenuSubsystem::UEditorMenuSubsystem() :
 	UEditorSubsystem(),
-	bRefreshWidgetsNextTick(false)
+	bNextTickTimerIsSet(false),
+	bRefreshWidgetsNextTick(false),
+	bCleanupStaleWidgetsNextTick(false)
 {
 }
 
 void UEditorMenuSubsystem::Initialize(FSubsystemCollectionBase& Collection)
 {
-
 }
 
 void UEditorMenuSubsystem::Deinitialize()
 {
-	bRefreshWidgetsNextTick = false;
 	GeneratedMenuWidgets.Reset();
 	WidgetObjectReferences.Reset();
 	Menus.Reset();
@@ -784,7 +784,7 @@ void UEditorMenuSubsystem::ExecuteStringCommand(const FEditorMenuStringCommand S
 	}
 }
 
-FUIAction UEditorMenuSubsystem::ConvertUIAction(const FEditorMenuStringCommand& StringCommand, const FEditorMenuContext& Context) const
+FUIAction UEditorMenuSubsystem::ConvertUIAction(const FEditorMenuStringCommand& StringCommand, const FEditorMenuContext& Context)
 {
 	FUIAction UIAction;
 
@@ -894,6 +894,8 @@ TSharedRef<SWidget> UEditorMenuSubsystem::GenerateWidget(const TArray<UEditorMen
 
 TSharedRef<SWidget> UEditorMenuSubsystem::GenerateWidget(UEditorMenu* GeneratedMenu)
 {
+	CleanupStaleWidgetsNextTick();
+
 	FGeneratedEditorMenuWidgets& WidgetsForMenuName = GeneratedMenuWidgets.FindOrAdd(GeneratedMenu->MenuName);
 
 	// Store a copy so that we can call 'Refresh' on menus not in the database
@@ -933,55 +935,73 @@ TSharedRef<SWidget> UEditorMenuSubsystem::GenerateWidget(UEditorMenu* GeneratedM
 	return SNullWidget::NullWidget;
 }
 
-void UEditorMenuSubsystem::RefreshAllWidgets()
+void UEditorMenuSubsystem::SetNextTickTimer()
 {
-	if (bRefreshWidgetsNextTick == false)
+	if (!bNextTickTimerIsSet)
 	{
-		bRefreshWidgetsNextTick = true;
-		GEditor->GetTimerManager()->SetTimerForNextTick(this, &UEditorMenuSubsystem::HandleRefreshAllWidgetsNextTick);
+		bNextTickTimerIsSet = true;
+		GEditor->GetTimerManager()->SetTimerForNextTick(this, &UEditorMenuSubsystem::HandleNextTick);
 	}
 }
 
-void UEditorMenuSubsystem::CleanupStaleGeneratedMenus()
+void UEditorMenuSubsystem::CleanupStaleWidgetsNextTick()
 {
-	for (auto WidgetsForMenuNameIt = GeneratedMenuWidgets.CreateIterator(); WidgetsForMenuNameIt; ++WidgetsForMenuNameIt)
+	bCleanupStaleWidgetsNextTick = true;
+	SetNextTickTimer();
+}
+
+void UEditorMenuSubsystem::RefreshAllWidgets()
+{
+	bRefreshWidgetsNextTick = true;
+	SetNextTickTimer();
+}
+
+void UEditorMenuSubsystem::HandleNextTick()
+{
+	if (bCleanupStaleWidgetsNextTick || bRefreshWidgetsNextTick)
 	{
-		FGeneratedEditorMenuWidgets& WidgetsForMenuName = WidgetsForMenuNameIt->Value;
+		CleanupStaleWidgets();
+		bCleanupStaleWidgetsNextTick = false;
 
-		for (auto Instance = WidgetsForMenuName.Instances.CreateIterator(); Instance; ++Instance)
+		if (bRefreshWidgetsNextTick)
 		{
-			if (!Instance->Widget.IsValid())
+			for (auto WidgetsForMenuNameIt = GeneratedMenuWidgets.CreateIterator(); WidgetsForMenuNameIt; ++WidgetsForMenuNameIt)
 			{
-				Instance.RemoveCurrent();
+				FGeneratedEditorMenuWidgets& WidgetsForMenuName = WidgetsForMenuNameIt->Value;
+				for (auto Instance = WidgetsForMenuName.Instances.CreateIterator(); Instance; ++Instance)
+				{
+					if (Instance->Widget.IsValid())
+					{
+						RefreshMenuWidget(WidgetsForMenuNameIt->Key, *Instance);
+					}
+				}
 			}
-		}
 
-		if (WidgetsForMenuName.Instances.Num() == 0)
-		{
-			WidgetsForMenuNameIt.RemoveCurrent();
+			bRefreshWidgetsNextTick = false;
 		}
 	}
+
+	bNextTickTimerIsSet = false;
 }
 
-void UEditorMenuSubsystem::HandleRefreshAllWidgetsNextTick()
+void UEditorMenuSubsystem::CleanupStaleWidgets()
 {
-	if (bRefreshWidgetsNextTick)
+	for (auto WidgetsForMenuNameIt = GeneratedMenuWidgets.CreateIterator(); WidgetsForMenuNameIt; ++WidgetsForMenuNameIt)
 	{
-		CleanupStaleGeneratedMenus();
+		FGeneratedEditorMenuWidgets& WidgetsForMenuName = WidgetsForMenuNameIt->Value;
 
-		for (auto WidgetsForMenuNameIt = GeneratedMenuWidgets.CreateIterator(); WidgetsForMenuNameIt; ++WidgetsForMenuNameIt)
+		for (auto Instance = WidgetsForMenuName.Instances.CreateIterator(); Instance; ++Instance)
 		{
-			FGeneratedEditorMenuWidgets& WidgetsForMenuName = WidgetsForMenuNameIt->Value;
-			for (auto Instance = WidgetsForMenuName.Instances.CreateIterator(); Instance; ++Instance)
+			if (!Instance->Widget.IsValid())
 			{
-				if (Instance->Widget.IsValid())
-				{
-					RefreshMenuWidget(WidgetsForMenuNameIt->Key, *Instance);
-				}
+				Instance.RemoveCurrent();
 			}
 		}
 
-		bRefreshWidgetsNextTick = false;
+		if (WidgetsForMenuName.Instances.Num() == 0)
+		{
+			WidgetsForMenuNameIt.RemoveCurrent();
+		}
 	}
 }
 
diff --git a/Engine/Source/Editor/EditorMenus/Public/EditorMenuEntry.h b/Engine/Source/Editor/EditorMenus/Public/EditorMenuEntry.h
--- a/Engine/Source/Editor/EditorMenus/Public/EditorMenuEntry.h
+++ b/Engine/Source/Editor/EditorMenus/Public/EditorMenuEntry.h
@@ -15,6 +15,8 @@
 
 #include "EditorMenuEntry.generated.h"
 
+class UEditorMenuEntryScript;
+
 USTRUCT(BlueprintType)
 struct EDITORMENUS_API FSubMenuEntryData
 {
diff --git a/Engine/Source/Editor/EditorMenus/Public/EditorMenuSection.h b/Engine/Source/Editor/EditorMenus/Public/EditorMenuSection.h
--- a/Engine/Source/Editor/EditorMenus/Public/EditorMenuSection.h
+++ b/Engine/Source/Editor/EditorMenus/Public/EditorMenuSection.h
@@ -24,7 +24,7 @@ public:
 USTRUCT(BlueprintType)
 struct EDITORMENUS_API FEditorMenuSection
 {
-	GENERATED_USTRUCT_BODY()
+	GENERATED_BODY()
 
 public:
 
@@ -34,7 +34,7 @@ public:
 
 	FEditorMenuEntry& AddEntryObject(UEditorMenuEntryScript* InObject);
 	FEditorMenuEntry& AddEntry(const FEditorMenuEntry& Args);
-	FEditorMenuEntry& AddMenuEntry(const TSharedPtr< const FUICommandInfo >& InCommand);
+	FEditorMenuEntry& AddMenuEntry(const TSharedPtr< const FUICommandInfo >& InCommand, const TAttribute<FText>& InLabelOverride = TAttribute<FText>(), const TAttribute<FText>& InToolTipOverride = TAttribute<FText>(), const TAttribute<FSlateIcon>& InIconOverride = TAttribute<FSlateIcon>(), const FName InTutorialHighlightName = NAME_None, const FName InNameOverride = NAME_None);
 
 	FEditorMenuEntry& AddDynamicEntry(const FName InName, const FNewEditorMenuSectionDelegate& InConstruct);
 	FEditorMenuEntry& AddDynamicEntry(const FName InName, const FNewEditorMenuDelegateLegacy& InConstruct);
@@ -47,6 +47,7 @@ public:
 	}
 
 private:
+
 	void InitGeneratedSectionCopy(const FEditorMenuSection& Source, FEditorMenuContext& InContext);
 
 	int32 RemoveEntry(const FName InName);
@@ -62,6 +63,7 @@ private:
 	friend class UEditorMenu;
 
 public:
+
 	UPROPERTY(EditAnywhere, BlueprintReadWrite, Category = "Editor UI")
 	FName Name;
 
diff --git a/Engine/Source/Editor/EditorMenus/Public/EditorMenuSubsystem.h b/Engine/Source/Editor/EditorMenus/Public/EditorMenuSubsystem.h
--- a/Engine/Source/Editor/EditorMenus/Public/EditorMenuSubsystem.h
+++ b/Engine/Source/Editor/EditorMenus/Public/EditorMenuSubsystem.h
@@ -37,7 +37,6 @@ class UEditorMenu;
 class UEditorMenuEntryScript;
 struct FEditorMenuEntry;
 struct FEditorMenuSection;
-//struct FEditorMenuOwnerScoped;
 
 USTRUCT()
 struct FCustomizedEditorMenuSection
@@ -98,6 +97,7 @@ public:
 
 	UEditorMenuSubsystem();
 
+	/** Get EditorMenuSubsystem if GEditor is valid and cause EditorMenus module to load. */
 	static inline UEditorMenuSubsystem* Get()
 	{
 		if (GEditor)
@@ -109,7 +109,8 @@ public:
 		return nullptr;
 	}
 
-	static inline UEditorMenuSubsystem* GetIfLoaded()
+	/** Try to get EditorMenuSubsystem if GEditor is valid without forcing EditorMenus module to load. */
+	static inline UEditorMenuSubsystem* TryGet()
 	{
 		if (GEditor && IEditorMenusModule::IsAvailable())
 		{
@@ -119,10 +120,10 @@ public:
 		return nullptr;
 	}
 
-	// Unregister all entries with a specific owner
+	/** Unregister everything associated with the given owner without forcing EditorMenus module to load. */
 	static inline void UnregisterOwner(FEditorMenuOwner Owner)
 	{
-		if (UEditorMenuSubsystem* EditorMenus = UEditorMenuSubsystem::GetIfLoaded())
+		if (UEditorMenuSubsystem* EditorMenus = UEditorMenuSubsystem::TryGet())
 		{
 			EditorMenus->UnregisterOwnerInternal(Owner);
 		}
@@ -133,66 +134,125 @@ public:
 	virtual void Deinitialize();
 
 
-	UEditorMenu* GenerateMenu(const FName Name, FEditorMenuContext& InMenuContext);
-	UEditorMenu* GenerateMenu(const TArray<UEditorMenu*>& Hierarchy, FEditorMenuContext& InMenuContext);
-	UEditorMenu* GenerateMenuAsBuilder(const UEditorMenu* InMenu, FEditorMenuContext& InMenuContext);
-
-	// Generate switch based on registered name
-	TSharedRef<SWidget> GenerateWidget(const FName Name, FEditorMenuContext& InMenuContext);
-	TSharedRef<SWidget> GenerateWidget(const TArray<UEditorMenu*>& Hierarchy, FEditorMenuContext& InMenuContext);
-	TSharedRef<SWidget> GenerateWidget(UEditorMenu* GeneratedMenu);
-
+	/**
+	 * Registers a menu by name
+	 * @param	Parent	Optional name of a menu to layer on top of.
+	 * @param	Type	Type of menu that will be generated such as: ToolBar, VerticalToolBar, etc..
+	 * @return	EditorMenu	Menu object
+	 */
 	UFUNCTION(BlueprintCallable, Category = "Editor UI")
 	UEditorMenu* RegisterMenu(FName Name, const FName Parent = NAME_None, EMultiBoxType Type = EMultiBoxType::Menu);
 
+	/**
+	 * Extends a menu without registering the menu or claiming ownership of it. Ok to call even if menu does not exist yet.
+	 * @param	Name	Name of the menu to extend
+	 * @return	EditorMenu	Menu object
+	 */
 	UFUNCTION(BlueprintCallable, Category = "Editor UI")
 	UEditorMenu* ExtendMenu(const FName Name);
 
+	/**
+	 * Generate widget from a registered menu. Most common function used to generate new menu widgets.
+	 * @param	Name	Registered menu's name that widget will be generated from
+	 * @param	Context	Additional information specific to the menu being generated
+	 * @return	Widget to display
+	 */
+	TSharedRef<SWidget> GenerateWidget(const FName Name, FEditorMenuContext& InMenuContext);
+
+
+
+
+	/**
+	 * Finds an existing menu that has been registered or extended.
+	 * @param	Name	Name of the menu to find.
+	 * @return	EditorMenu	Menu object. Returns null if not found.
+	 */
 	UFUNCTION(BlueprintCallable, Category = "Editor UI")
 	UEditorMenu* FindMenu(const FName Name);
 
-	// Rebuilds all widgets generated from a specific menu
+	/** Rebuilds all widgets generated from a specific menu. */
 	UFUNCTION(BlueprintCallable, Category = "Editor UI")
 	bool RefreshMenuWidget(const FName Name);
-	
-	// Rebuilds all generated widgets next tick
+
+	/** Rebuilds all currently generated widgets next tick. */
 	UFUNCTION(BlueprintCallable, Category = "Editor UI")
 	void RefreshAllWidgets();
 
+	/** Registers menu entry object from blueprint/script */
 	UFUNCTION(BlueprintCallable, Category = "Editor UI")
 	static bool AddMenuEntryObject(UEditorMenuEntryScript* MenuEntryObject);
 
+	/** Removes all entries that were registered under a specific owner name */
 	UFUNCTION(BlueprintCallable, Category = "Editor UI")
 	void UnregisterOwnerByName(FName InOwnerName);
 
+	/** Sets a section's displayed label text. */
 	UFUNCTION(BlueprintCallable, Category = "Editor UI")
 	void SetSectionLabel(const FName MenuName, const FName SectionName, const FText Label);
 
+	/** Sets where to insert a section into a menu when generating relative to other section names. */
 	UFUNCTION(BlueprintCallable, Category = "Editor UI")
-	void SetSectionPosition(const FName MenuName, const FName SectionName, const FName PositionName, const EEditorMenuInsertType PositionType);
+	void SetSectionPosition(const FName MenuName, const FName SectionName, const FName OtherSectionName, const EEditorMenuInsertType PositionType);
 
+	/** Registers a section for a menu */
 	void AddSection(const FName MenuName, const FName SectionName, const TAttribute< FText >& InLabel, const FEditorMenuInsert InPosition = FEditorMenuInsert());
 
+	/** Registers an entry for a menu's section */
 	void AddEntry(const FName MenuName, const FName SectionName, const FEditorMenuEntry& Entry);
 
+	/** Removes a menu entry from a given menu and section */
 	UFUNCTION(BlueprintCallable, Category = "Editor UI")
 	void RemoveEntry(const FName MenuName, const FName Section, const FName Name);
 
+	/** Removes a section from a given menu */
 	UFUNCTION(BlueprintCallable, Category = "Editor UI")
 	void RemoveSection(const FName MenuName, const FName Section);
 
+	/** Unregisters a menu by name */
 	UFUNCTION(BlueprintCallable, Category = "Editor UI")
 	void RemoveMenu(const FName MenuName);
 
+	/** Finds a context object of a given class if it exists */
 	UFUNCTION(BlueprintCallable, Category = "Editor UI")
 	static UObject* FindContext(const FEditorMenuContext& InContext, UClass* InClass);
 
+	/**
+	 * Generate widget from a hierarchy of menus. For advanced specialized use cases.
+	 * @param	Hierarchy	Array of menus to combine into a final widget
+	 * @param	Context	Additional information specific to the menu being generated
+	 * @return	Widget to display
+	 */
+	TSharedRef<SWidget> GenerateWidget(const TArray<UEditorMenu*>& Hierarchy, FEditorMenuContext& InMenuContext);
+
+	/**
+	 * Generate widget from a final collapsed menu. For advanced specialized use cases.
+	 * @param	GeneratedMenu	Combined final menu to generate final widget from
+	 * @return	Widget to display
+	 */
+	TSharedRef<SWidget> GenerateWidget(UEditorMenu* GeneratedMenu);
+	
+	/** Create a finalized menu that combines all parents used to generate a widget. Advanced special use cases only. */
+	UEditorMenu* GenerateMenu(const FName Name, FEditorMenuContext& InMenuContext);
+
+	/** Create a finalized menu that combines given hierarchy array that will generate a widget. Advanced special use cases only. */
+	UEditorMenu* GenerateMenu(const TArray<UEditorMenu*>& Hierarchy, FEditorMenuContext& InMenuContext);
+
+	/** Create a finalized menu based on a custom crafted menu. Advanced special use cases only. */
+	UEditorMenu* GenerateMenuAsBuilder(const UEditorMenu* InMenu, FEditorMenuContext& InMenuContext);
+
+	/** For advanced use cases */
 	void AssembleMenuByName(UEditorMenu* GeneratedMenu, const FName Name);
+
+	/** For advanced use cases */
 	void AssembleMenuHierarchy(UEditorMenu* GeneratedMenu, const TArray<UEditorMenu*>& Hierarchy);
 
+	/** For advanced use cases */
 	FEditorMenuOwner CurrentOwner() const;
 
+	/** Registers a new type of string based command handler. */
 	void RegisterStringCommandHandler(const FName InName, const FEditorMenuExecuteString& InDelegate);
+
+	/** Removes a string based command handler. */
 	void UnregisterStringCommandHandler(const FName InName);
 
 	friend struct FEditorMenuOwnerScoped;
@@ -203,20 +263,31 @@ public:
 
 private:
 
-	void HandleRefreshAllWidgetsNextTick();
+	/** Sets a timer to be called next engine tick so that multiple repeated actions can be combined together. */
+	void SetNextTickTimer();
 
-	bool RefreshMenuWidget(const FName Name, FGeneratedEditorMenuWidget& GeneratedMenuWidget);
-	
-	void CleanupStaleGeneratedMenus();
+	/** Timer function used to consolidate multiple duplicate requests into a single frame. */
+	void HandleNextTick();
 
-	TArray<UEditorMenu*> CollectHierarchy(const FName Name);
+	/** Release references to UObjects of widgets that have been deleted. Combines multiple requests in one frame together for improved performance. */
+	void CleanupStaleWidgetsNextTick();
 
+	/** Release references to UObjects of widgets that have been deleted */
+	void CleanupStaleWidgets();
+
+	/** Re-creates widget that is active */
+	bool RefreshMenuWidget(const FName Name, FGeneratedEditorMenuWidget& GeneratedMenuWidget);
+
+	/** Sets the current temporary menu owner to avoid needing to supply owner for each menu entry being registered. Used by FEditorMenuEntryScoped */
 	void PushOwner(const FEditorMenuOwner InOwner);
 
+	/** Sets the current temporary menu owner. Used by FEditorMenuEntryScoped */
 	void PopOwner(const FEditorMenuOwner InOwner);
 
 	UEditorMenu* FindSubMenuToGenerateWith(const FName InParentName, const FName InChildName);
 
+	TArray<UEditorMenu*> CollectHierarchy(const FName Name);
+
 	void FillMenu(FMenuBuilder& MenuBuilder, FName InMenuName, FEditorMenuContext InMenuContext);
 	void FillMenuBarDropDown(FMenuBuilder& MenuBuilder, FName InParentName, FName InChildName, FEditorMenuContext InMenuContext);
 	void PopulateMenuBuilder(FMenuBuilder& MenuBuilder, UEditorMenu* MenuData);
@@ -225,7 +296,8 @@ private:
 
 	TSharedRef<SWidget> GenerateToolbarComboButtonMenu(const FName SubMenuFullName, FEditorMenuContext InContext);
 
-	FUIAction ConvertUIAction(const FEditorMenuStringCommand& StringCommand, const FEditorMenuContext& Context) const;
+	/** Converts a string command to a FUIAction */
+	FUIAction ConvertUIAction(const FEditorMenuStringCommand& StringCommand, const FEditorMenuContext& Context);
 	static FUIAction ConvertUIAction(FEditorUIActionChoice& Actions, FEditorMenuContext& Context);
 	static FUIAction ConvertUIAction(FEditorUIAction& Actions, FEditorMenuContext& Context);
 	static FUIAction ConvertUIAction(FEditorDynamicUIAction& Actions, FEditorMenuContext& Context);
@@ -255,13 +327,11 @@ private:
 
 	static bool GetDisplayUIExtensionPoints();
 
-public:
+private:
 
 	UPROPERTY(config, EditAnywhere, Category = Misc)
 	TArray<FCustomizedEditorMenu> CustomizedMenus;
 
-private:
-
 	UPROPERTY()
 	TMap<FName, UEditorMenu*> Menus;
 
@@ -274,7 +344,9 @@ private:
 
 	TMap<FName, FEditorMenuExecuteString> StringCommandHandlers;
 
+	bool bNextTickTimerIsSet;
 	bool bRefreshWidgetsNextTick;
+	bool bCleanupStaleWidgetsNextTick;
 };
 
 struct FEditorMenuOwnerScoped
