Fix static analysis warnings
#rnx

diff --git a/Engine/Source/Editor/EditorMenus/Private/EditorMenuSubsystem.cpp b/Engine/Source/Editor/EditorMenus/Private/EditorMenuSubsystem.cpp
--- a/Engine/Source/Editor/EditorMenus/Private/EditorMenuSubsystem.cpp
+++ b/Engine/Source/Editor/EditorMenus/Private/EditorMenuSubsystem.cpp
@@ -921,7 +921,7 @@ TSharedRef<SWidget> UEditorMenuSubsystem::GenerateWidget(UEditorMenu* GeneratedM
 		GeneratedMenuWidget.Widget = Result;
 		return Result;
 	}
-	else if (GeneratedMenu->MenuType == EMultiBoxType::ToolBar)
+	else if (GeneratedMenu->MenuType == EMultiBoxType::ToolBar || GeneratedMenu->MenuType == EMultiBoxType::VerticalToolBar)
 	{
 		const EOrientation ToolBarOrientation = (GeneratedMenu->MenuType == EMultiBoxType::VerticalToolBar) ? Orient_Vertical : Orient_Horizontal;
 		FToolBarBuilder ToolbarBuilder(GeneratedMenu->Context.CommandList, GeneratedMenu->MenuName, GeneratedMenu->Context.GetAllExtenders(), ToolBarOrientation, GeneratedMenu->bToolBarForceSmallIcons);
@@ -1050,7 +1050,7 @@ bool UEditorMenuSubsystem::RefreshMenuWidget(const FName InName, FGeneratedEdito
 		PopulateMenuBarBuilder(MenuBarBuilder, GeneratedMenu);
 		MultiBoxWidget->SetMultiBox(MenuBarBuilder.GetMultiBox());
 	}
-	else if (GeneratedMenu->MenuType == EMultiBoxType::ToolBar)
+	else if (GeneratedMenu->MenuType == EMultiBoxType::ToolBar || GeneratedMenu->MenuType == EMultiBoxType::VerticalToolBar)
 	{
 		const EOrientation ToolBarOrientation = (GeneratedMenu->MenuType == EMultiBoxType::VerticalToolBar) ? Orient_Vertical : Orient_Horizontal;
 		FToolBarBuilder ToolbarBuilder(GeneratedMenu->Context.CommandList, GeneratedMenu->MenuName, GeneratedMenu->Context.GetAllExtenders(), ToolBarOrientation, GeneratedMenu->bToolBarForceSmallIcons);
