Editor menu updates:
- Widget menu entries
- StringCommand enums
- Show extension points now can copy menu names to clipboard when menu name entry is clicked

#rnx

diff --git a/Engine/Source/Editor/EditorMenus/EditorMenus.Build.cs b/Engine/Source/Editor/EditorMenus/EditorMenus.Build.cs
--- a/Engine/Source/Editor/EditorMenus/EditorMenus.Build.cs
+++ b/Engine/Source/Editor/EditorMenus/EditorMenus.Build.cs
@@ -10,6 +10,7 @@ namespace UnrealBuildTool.Rules
 				new string[] {
 					"Core",
 					"CoreUObject",
+					"ApplicationCore",
 					"Slate",
 					"SlateCore",
 					"EditorSubSystem",
diff --git a/Engine/Source/Editor/EditorMenus/Private/EditorMenuEntry.cpp b/Engine/Source/Editor/EditorMenus/Private/EditorMenuEntry.cpp
--- a/Engine/Source/Editor/EditorMenus/Private/EditorMenuEntry.cpp
+++ b/Engine/Source/Editor/EditorMenus/Private/EditorMenuEntry.cpp
@@ -33,8 +33,6 @@ FEditorMenuEntry::FEditorMenuEntry(const FEditorMenuOwner InOwner, const FName I
 
 void FEditorMenuEntry::SetCommand(const TSharedPtr< const FUICommandInfo >& InCommand, FName InName, const TAttribute<FText>& InLabel, const TAttribute<FText>& InToolTip, const TAttribute<FSlateIcon>& InIcon)
 {
-	//Action = FEditorUIActionChoice(InCommand, InCommandList);
-
 	Command = InCommand;
 	Name = InName != NAME_None ? InName : InCommand->GetCommandName();
 	Label = InLabel.IsSet() ? InLabel : InCommand->GetLabel();
@@ -62,6 +60,14 @@ FEditorMenuEntry FEditorMenuEntry::InitMenuEntry(const TSharedPtr< const FUIComm
 	return Entry;
 }
 
+FEditorMenuEntry FEditorMenuEntry::InitMenuEntry(const FName InName, const FEditorUIActionChoice& InAction, const TSharedRef<SWidget>& Widget)
+{
+	FEditorMenuEntry Entry(UEditorMenuSubsystem::Get()->CurrentOwner(), InName, EMultiBlockType::MenuEntry);
+	Entry.Action = InAction;
+	Entry.MakeWidget.BindLambda([Widget](const FEditorMenuContext&) { return Widget; });
+	return Entry;
+}
+
 FEditorMenuEntry FEditorMenuEntry::InitSubMenu(const FName InParentMenu, const FName InName, const TAttribute<FText>& InLabel, const TAttribute<FText>& InToolTip, const FNewEditorMenuChoice& InMakeMenu, bool bInOpenSubMenuOnClick, const TAttribute<FSlateIcon>& InIcon, const bool bInShouldCloseWindowAfterMenuSelection)
 {
 	FEditorMenuEntry Entry(UEditorMenuSubsystem::Get()->CurrentOwner(), InName, EMultiBlockType::MenuEntry);
@@ -94,7 +100,7 @@ FEditorMenuEntry FEditorMenuEntry::InitToolBarButton(const TSharedPtr< const FUI
 	return Entry;
 }
 
-FEditorMenuEntry FEditorMenuEntry::InitComboButton(const FName InName, const FEditorUIActionChoice& InAction, const FOnGetContent& InMenuContentGenerator, const TAttribute<FText>& InLabel, const TAttribute<FText>& InToolTip, const TAttribute<FSlateIcon>& InIcon, bool bInSimpleComboBox, FName InTutorialHighlightName)
+FEditorMenuEntry FEditorMenuEntry::InitComboButton(const FName InName, const FEditorUIActionChoice& InAction, const FNewEditorMenuWidgetChoice& InMenuContentGenerator, const TAttribute<FText>& InLabel, const TAttribute<FText>& InToolTip, const TAttribute<FSlateIcon>& InIcon, bool bInSimpleComboBox, FName InTutorialHighlightName)
 {
 	FEditorMenuEntry Entry(UEditorMenuSubsystem::Get()->CurrentOwner(), InName, EMultiBlockType::ToolBarComboButton);
 	Entry.TutorialHighlightName = InTutorialHighlightName;
@@ -121,8 +127,9 @@ FEditorMenuEntry FEditorMenuEntry::InitToolBarSeparator(const FName InName)
 void FEditorMenuEntry::ResetActions()
 {
 	Action = FEditorUIActionChoice();
-	//ScriptObject = nullptr;
-	StringCommand = FEditorMenuStringCommand();
 	Command.Reset();
 	CommandList.Reset();
+	StringExecuteAction = FEditorMenuStringCommand();
+	// Note: Cannot reset ScriptObject as it would also remove label and other data
+	//ScriptObject = nullptr;
 }
diff --git a/Engine/Source/Editor/EditorMenus/Private/EditorMenuSection.cpp b/Engine/Source/Editor/EditorMenus/Private/EditorMenuSection.cpp
--- a/Engine/Source/Editor/EditorMenus/Private/EditorMenuSection.cpp
+++ b/Engine/Source/Editor/EditorMenus/Private/EditorMenuSection.cpp
@@ -71,6 +71,11 @@ FEditorMenuEntry& FEditorMenuSection::AddEntryObject(UEditorMenuEntryScript* InO
 	return AddEntry(Args);
 }
 
+FEditorMenuEntry& FEditorMenuSection::AddMenuEntry(const FName InName, const TAttribute<FText>& InLabel, const TAttribute<FText>& InToolTip, const TAttribute<FSlateIcon>& InIcon, const FEditorUIActionChoice& InAction, const EUserInterfaceActionType InUserInterfaceActionType, const FName InTutorialHighlightName)
+{
+	return AddEntry(FEditorMenuEntry::InitMenuEntry(InName, InLabel, InToolTip, InIcon, InAction, InUserInterfaceActionType, InTutorialHighlightName));
+}
+
 FEditorMenuEntry& FEditorMenuSection::AddMenuEntry(const TSharedPtr< const FUICommandInfo >& InCommand, const TAttribute<FText>& InLabelOverride, const TAttribute<FText>& InToolTipOverride, const TAttribute<FSlateIcon>& InIconOverride, const FName InTutorialHighlightName, const FName InNameOverride)
 {
 	return AddEntry(FEditorMenuEntry::InitMenuEntry(InCommand, InLabelOverride, InToolTipOverride, InIconOverride, InTutorialHighlightName, InNameOverride));
diff --git a/Engine/Source/Editor/EditorMenus/Private/EditorMenuSubsystem.cpp b/Engine/Source/Editor/EditorMenus/Private/EditorMenuSubsystem.cpp
--- a/Engine/Source/Editor/EditorMenus/Private/EditorMenuSubsystem.cpp
+++ b/Engine/Source/Editor/EditorMenus/Private/EditorMenuSubsystem.cpp
@@ -9,6 +9,7 @@
 #include "Framework/MultiBox/MultiBoxBuilder.h"
 #include "Internationalization/Internationalization.h"
 #include "Editor/EditorPerProjectUserSettings.h"
+#include "HAL/PlatformApplicationMisc.h"
 
 #include "Editor.h"
 
@@ -21,6 +22,36 @@ FAutoConsoleCommand EditorMenusRefreshMenuWidget = FAutoConsoleCommand(
 		UEditorMenuSubsystem::Get()->RefreshAllWidgets();
 	}));
 
+FName FEditorMenuStringCommand::GetTypeName() const
+{
+	static const FName CommandName("Command");
+	static const FName PythonName("Python");
+
+	switch (Type)
+	{
+	case EEditorMenuStringCommandType::Command:
+		return CommandName;
+	case EEditorMenuStringCommandType::Python:
+		return PythonName;
+	case EEditorMenuStringCommandType::Other:
+		return TypeName;
+	default:
+		break;
+	}
+
+	return NAME_None;
+}
+
+FExecuteAction FEditorMenuStringCommand::ToExecuteAction(const FEditorMenuContext& Context) const
+{
+	if (IsBound())
+	{
+		return FExecuteAction::CreateStatic(&UEditorMenuSubsystem::ExecuteStringCommand, *this, Context);
+	}
+
+	return FExecuteAction();
+}
+
 FEditorUIActionChoice::FEditorUIActionChoice(const TSharedPtr< const FUICommandInfo >& InCommand, const FUICommandList& InCommandList)
 {
 	if (InCommand.IsValid())
@@ -69,6 +100,12 @@ UEditorMenu* UEditorMenuSubsystem::FindMenu(const FName Name)
 	return Found ? *Found : nullptr;
 }
 
+bool UEditorMenuSubsystem::IsMenuRegistered(const FName Name) const
+{
+	const UEditorMenu* const * Found = Menus.Find(Name);
+	return Found && *Found && (*Found)->IsRegistered();
+}
+
 TArray<UEditorMenu*> UEditorMenuSubsystem::CollectHierarchy(const FName InName)
 {
 	TArray<UEditorMenu*> Result;
@@ -429,6 +466,7 @@ void UEditorMenuSubsystem::FillMenuBarDropDown(class FMenuBuilder& MenuBuilder,
 		UEditorMenu* GeneratedMenu = NewObject<UEditorMenu>();
 		GeneratedMenu->Context = InMenuContext;
 		AssembleMenuByName(GeneratedMenu, MenuToUse->MenuName);
+		GeneratedMenu->MenuName = JoinMenuPaths(InParentName, InChildName);
 
 		// Populate menu builder with final menu
 		PopulateMenuBuilder(MenuBuilder, GeneratedMenu);
@@ -439,7 +477,14 @@ void UEditorMenuSubsystem::PopulateMenuBuilder(FMenuBuilder& MenuBuilder, UEdito
 {
 	if (GetDisplayUIExtensionPoints())
 	{
-		MenuBuilder.AddMenuEntry(FText::FromName(MenuData->GetMenuName()), TAttribute<FText>(), FSlateIcon(), FUIAction(), "MenuName");
+		const FName MenuName = MenuData->GetMenuName();
+		MenuBuilder.AddMenuEntry(
+			FText::FromName(MenuName),
+			LOCTEXT("CopyMenuNameToClipboard", "Copy menu name to clipboard"),
+			FSlateIcon(),
+			FExecuteAction::CreateLambda([MenuName]() { FPlatformApplicationMisc::ClipboardCopy(*MenuName.ToString()); }),
+			"MenuName"
+		);
 	}
 
 	for (FEditorMenuSection& Section : MenuData->Sections)
@@ -457,8 +502,18 @@ void UEditorMenuSubsystem::PopulateMenuBuilder(FMenuBuilder& MenuBuilder, UEdito
 			if (Block.ConstructLegacy.IsBound())
 			{
 				Block.ConstructLegacy.Execute(MenuBuilder, MenuData);
+				continue;
+			}
+
+			FUIAction UIAction = ConvertUIAction(Block, MenuData->Context);
+
+			TSharedPtr<SWidget> Widget;
+			if (Block.MakeWidget.IsBound())
+			{
+				Widget = Block.MakeWidget.Execute(MenuData->Context);
 			}
-			else if (Block.Type == EMultiBlockType::MenuEntry)
+
+			if (Block.Type == EMultiBlockType::MenuEntry)
 			{
 				if (Block.SubMenu.bIsSubMenu)
 				{
@@ -491,38 +546,55 @@ void UEditorMenuSubsystem::PopulateMenuBuilder(FMenuBuilder& MenuBuilder, UEdito
 					{
 						// SubMenu registered once by name in database
 						FName SubMenuFullName = JoinMenuPaths(MenuData->MenuName, Block.Name);
-						MenuBuilder.AddSubMenu(
-							Block.Label,
-							Block.ToolTip,
-							FNewMenuDelegate::CreateUObject(this, &UEditorMenuSubsystem::FillMenu, SubMenuFullName, MenuData->Context),
-							Block.SubMenu.bOpenSubMenuOnClick,
-							Block.Icon.Get(),
-							Block.bShouldCloseWindowAfterMenuSelection,
-							Block.Name
-						);
+						FNewMenuDelegate NewMenuDelegate = FNewMenuDelegate::CreateUObject(this, &UEditorMenuSubsystem::FillMenu, SubMenuFullName, MenuData->Context);
+
+						if (Widget.IsValid())
+						{
+							// Could also check if Visible/Enabled bound as well
+							if (UIAction.IsBound())
+							{
+								MenuBuilder.AddSubMenu(UIAction, Widget.ToSharedRef(), NewMenuDelegate, Block.bShouldCloseWindowAfterMenuSelection);
+							}
+							else
+							{
+								MenuBuilder.AddSubMenu(Widget.ToSharedRef(), NewMenuDelegate, Block.SubMenu.bOpenSubMenuOnClick, Block.bShouldCloseWindowAfterMenuSelection);
+							}
+						}
+						else
+						{
+							MenuBuilder.AddSubMenu(
+								Block.Label,
+								Block.ToolTip,
+								NewMenuDelegate,
+								Block.SubMenu.bOpenSubMenuOnClick,
+								Block.Icon.Get(),
+								Block.bShouldCloseWindowAfterMenuSelection,
+								Block.Name
+							);
+						}
 					}
 				}
 				else
 				{
-					if (Block.StringCommand.String.Len() > 0)
-					{
-						FUIAction UIAction = ConvertUIAction(Block.StringCommand, MenuData->Context);
-						MenuBuilder.AddMenuEntry(Block.Label, Block.ToolTip, Block.Icon.Get(), UIAction, Block.Name, Block.UserInterfaceActionType, Block.TutorialHighlightName);
-					}
-					else if (Block.Command.IsValid())
+					if (Block.Command.IsValid())
 					{
 						MenuBuilder.AddMenuEntry(Block.Command, Block.Name, Block.Label, Block.ToolTip, Block.Icon.Get());
 					}
 					else if (Block.ScriptObject)
 					{
 						UEditorMenuEntryScript* ScriptObject = Block.ScriptObject;
-						FUIAction UIAction = ConvertUIAction(ScriptObject, MenuData->Context);
 						MenuBuilder.AddMenuEntry(ScriptObject->CreateLabelAttribute(MenuData->Context), ScriptObject->CreateToolTipAttribute(MenuData->Context), Block.Icon.Get(), UIAction, ScriptObject->Data.Name, Block.UserInterfaceActionType, Block.TutorialHighlightName);
 					}
 					else
 					{
-						FUIAction UIAction = ConvertUIAction(Block.Action, MenuData->Context);
-						MenuBuilder.AddMenuEntry(Block.Label, Block.ToolTip, Block.Icon.Get(), UIAction, Block.Name, Block.UserInterfaceActionType, Block.TutorialHighlightName);
+						if (Widget.IsValid())
+						{
+							MenuBuilder.AddMenuEntry(UIAction, Widget.ToSharedRef(), Block.Name, Block.ToolTip, Block.UserInterfaceActionType, Block.TutorialHighlightName);
+						}
+						else
+						{
+							MenuBuilder.AddMenuEntry(Block.Label, Block.ToolTip, Block.Icon.Get(), UIAction, Block.Name, Block.UserInterfaceActionType, Block.TutorialHighlightName);
+						}
 					}
 				}
 			}
@@ -555,41 +627,38 @@ void UEditorMenuSubsystem::PopulateToolBarBuilder(FToolBarBuilder& ToolBarBuilde
 			if (Block.ToolBar.ConstructLegacy.IsBound())
 			{
 				Block.ToolBar.ConstructLegacy.Execute(ToolBarBuilder, MenuData);
+				continue;
 			}
-			else if (Block.Type == EMultiBlockType::ToolBarButton)
+
+			FUIAction UIAction = ConvertUIAction(Block, MenuData->Context);
+			if (Block.Type == EMultiBlockType::ToolBarButton)
 			{
-				if (Block.StringCommand.String.Len() > 0)
-				{
-					FUIAction UIAction = ConvertUIAction(Block.StringCommand, MenuData->Context);
-					ToolBarBuilder.AddToolBarButton(UIAction, Block.Name, Block.Label, Block.ToolTip, Block.Icon, Block.UserInterfaceActionType, Block.TutorialHighlightName);
-				}
-				else if (Block.Command.IsValid())
+				if (Block.Command.IsValid())
 				{
 					ToolBarBuilder.AddToolBarButton(Block.Command, Block.Name, Block.Label, Block.ToolTip, Block.Icon, Block.TutorialHighlightName);
 				}
 				else if (Block.ScriptObject)
 				{
 					UEditorMenuEntryScript* ScriptObject = Block.ScriptObject;
-					FUIAction UIAction = ConvertUIAction(ScriptObject, MenuData->Context);
 					ToolBarBuilder.AddToolBarButton(UIAction, ScriptObject->Data.Name, ScriptObject->CreateLabelAttribute(MenuData->Context), ScriptObject->CreateToolTipAttribute(MenuData->Context), Block.Icon, Block.UserInterfaceActionType, Block.TutorialHighlightName);
 				}
 				else
 				{
-					FUIAction UIAction = ConvertUIAction(Block.Action, MenuData->Context);
 					ToolBarBuilder.AddToolBarButton(UIAction, Block.Name, Block.Label, Block.ToolTip, Block.Icon, Block.UserInterfaceActionType, Block.TutorialHighlightName);
 				}
 			}
 			else if (Block.Type == EMultiBlockType::ToolBarComboButton)
 			{
-				if (FUIAction* UIAction = Block.Action.GetUIAction())
+				FOnGetContent OnGetContent = ConvertWidgetChoice(Block.ToolBar.ComboButtonContextMenuGenerator, MenuData->Context);
+				if (OnGetContent.IsBound())
 				{
-					ToolBarBuilder.AddComboButton(*UIAction, Block.ToolBar.ComboButtonContextMenuGenerator, Block.Label, Block.ToolTip, Block.Icon, Block.ToolBar.bSimpleComboBox, Block.TutorialHighlightName);
+					ToolBarBuilder.AddComboButton(UIAction, OnGetContent, Block.Label, Block.ToolTip, Block.Icon, Block.ToolBar.bSimpleComboBox, Block.TutorialHighlightName);
 				}
 				else
 				{
 					FName SubMenuFullName = JoinMenuPaths(MenuData->MenuName, Block.Name);
 					FOnGetContent Delegate = FOnGetContent::CreateUObject(this, &UEditorMenuSubsystem::GenerateToolbarComboButtonMenu, SubMenuFullName, MenuData->Context);
-					ToolBarBuilder.AddComboButton(FUIAction(), Delegate, Block.Label, Block.ToolTip, Block.Icon, Block.ToolBar.bSimpleComboBox, Block.TutorialHighlightName);
+					ToolBarBuilder.AddComboButton(UIAction, Delegate, Block.Label, Block.ToolTip, Block.Icon, Block.ToolBar.bSimpleComboBox, Block.TutorialHighlightName);
 				}
 			}
 			else if (Block.Type == EMultiBlockType::ToolBarSeparator)
@@ -605,6 +674,19 @@ void UEditorMenuSubsystem::PopulateToolBarBuilder(FToolBarBuilder& ToolBarBuilde
 		ToolBarBuilder.EndSection();
 	}
 
+	if (GetDisplayUIExtensionPoints())
+	{
+		const FName MenuName = MenuData->GetMenuName();
+		ToolBarBuilder.BeginSection(MenuName);
+		ToolBarBuilder.AddToolBarButton(
+			FExecuteAction::CreateLambda([MenuName]() { FPlatformApplicationMisc::ClipboardCopy(*MenuName.ToString()); }), 
+			"MenuName",
+			LOCTEXT("CopyNameToClipboard", "Copy Name"),
+			LOCTEXT("CopyMenuNameToClipboard", "Copy menu name to clipboard")
+		);
+		ToolBarBuilder.EndSection();
+	}
+
 	AddReferencedContextObjects(ToolBarBuilder.GetMultiBox(), MenuData->Context);
 }
 
@@ -626,17 +708,70 @@ void UEditorMenuSubsystem::PopulateMenuBarBuilder(FMenuBarBuilder& MenuBarBuilde
 	}
 }
 
-FUIAction UEditorMenuSubsystem::ConvertUIAction(FEditorUIActionChoice& Choice, FEditorMenuContext& Context)
+FOnGetContent UEditorMenuSubsystem::ConvertWidgetChoice(const FNewEditorMenuWidgetChoice& Choice, const FEditorMenuContext& Context) const
 {
-	if (FEditorUIAction* EditorAction = Choice.GetEditorUIAction())
+	if (Choice.NewEditorMenuWidget.IsBound())
+	{
+		return FOnGetContent::CreateLambda([ToCall = Choice.NewEditorMenuWidget, Context]()
+		{
+			if (ToCall.IsBound())
+			{
+				return ToCall.Execute(Context);
+			}
+
+			return SNullWidget::NullWidget;
+		});
+	}
+	else if (Choice.NewEditorMenu.IsBound())
+	{
+		return FOnGetContent::CreateLambda([ToCall = Choice.NewEditorMenu, Context]()
+		{
+			if (ToCall.IsBound())
+			{
+				UEditorMenu* MenuData = NewObject<UEditorMenu>();
+				MenuData->Context = Context;
+				ToCall.Execute(MenuData);
+				return UEditorMenuSubsystem::Get()->GenerateWidget(MenuData);
+			}
+
+			return SNullWidget::NullWidget;
+		});
+	}
+	return Choice.OnGetContent;
+}
+
+FUIAction UEditorMenuSubsystem::ConvertUIAction(const FEditorMenuEntry& Block, const FEditorMenuContext& Context)
+{
+	FUIAction UIAction;
+	
+	if (Block.ScriptObject)
+	{
+		UIAction = ConvertScriptObjectToUIAction(Block.ScriptObject, Context);
+	}
+	else
+	{
+		UIAction = ConvertUIAction(Block.Action, Context);
+	}
+	
+	if (!UIAction.ExecuteAction.IsBound() && Block.StringExecuteAction.IsBound())
+	{
+		UIAction.ExecuteAction = Block.StringExecuteAction.ToExecuteAction(Context);
+	}
+
+	return UIAction;
+}
+
+FUIAction UEditorMenuSubsystem::ConvertUIAction(const FEditorUIActionChoice& Choice, const FEditorMenuContext& Context)
+{
+	if (const FEditorUIAction* EditorAction = Choice.GetEditorUIAction())
 	{
 		return ConvertUIAction(*EditorAction, Context);
 	}
-	else if (FEditorDynamicUIAction* DynamicEditorAction = Choice.GetEditorDynamicUIAction())
+	else if (const FEditorDynamicUIAction* DynamicEditorAction = Choice.GetEditorDynamicUIAction())
 	{
 		return ConvertUIAction(*DynamicEditorAction, Context);
 	}
-	else if (FUIAction* Action = Choice.GetUIAction())
+	else if (const FUIAction* Action = Choice.GetUIAction())
 	{
 		return *Action;
 	}
@@ -644,7 +779,7 @@ FUIAction UEditorMenuSubsystem::ConvertUIAction(FEditorUIActionChoice& Choice, F
 	return FUIAction();
 }
 
-FUIAction UEditorMenuSubsystem::ConvertUIAction(FEditorUIAction& Actions, FEditorMenuContext& Context)
+FUIAction UEditorMenuSubsystem::ConvertUIAction(const FEditorUIAction& Actions, const FEditorMenuContext& Context)
 {
 	FUIAction UIAction;
 
@@ -683,7 +818,7 @@ FUIAction UEditorMenuSubsystem::ConvertUIAction(FEditorUIAction& Actions, FEdito
 	return UIAction;
 }
 
-FUIAction UEditorMenuSubsystem::ConvertUIAction(FEditorDynamicUIAction& Actions, FEditorMenuContext& Context)
+FUIAction UEditorMenuSubsystem::ConvertUIAction(const FEditorDynamicUIAction& Actions, const FEditorMenuContext& Context)
 {
 	FUIAction UIAction;
 
@@ -722,7 +857,7 @@ FUIAction UEditorMenuSubsystem::ConvertUIAction(FEditorDynamicUIAction& Actions,
 	return UIAction;
 }
 
-FUIAction UEditorMenuSubsystem::ConvertUIAction(UEditorMenuEntryScript* ScriptObject, FEditorMenuContext& Context)
+FUIAction UEditorMenuSubsystem::ConvertScriptObjectToUIAction(UEditorMenuEntryScript* ScriptObject, const FEditorMenuContext& Context)
 {
 	FUIAction UIAction;
 
@@ -760,17 +895,20 @@ FUIAction UEditorMenuSubsystem::ConvertUIAction(UEditorMenuEntryScript* ScriptOb
 
 void UEditorMenuSubsystem::ExecuteStringCommand(const FEditorMenuStringCommand StringCommand, const FEditorMenuContext Context)
 {
-	if (StringCommand.String.Len() > 0)
+	if (StringCommand.IsBound())
 	{
 		static const FName CommandTypeName("Command");
-		if (const FEditorMenuExecuteString* Handler = StringCommandHandlers.Find(StringCommand.TypeName))
+		const FName TypeName = StringCommand.GetTypeName();
+
+		UEditorMenuSubsystem* EditorMenus = UEditorMenuSubsystem::Get();
+		if (const FEditorMenuExecuteString* Handler = EditorMenus->StringCommandHandlers.Find(TypeName))
 		{
 			if (Handler->IsBound())
 			{
 				Handler->Execute(StringCommand.String, Context);
 			}
 		}
-		else if (StringCommand.TypeName == CommandTypeName)
+		else if (TypeName == CommandTypeName)
 		{
 			if (GEditor)
 			{
@@ -779,23 +917,11 @@ void UEditorMenuSubsystem::ExecuteStringCommand(const FEditorMenuStringCommand S
 		}
 		else
 		{
-			UE_LOG(LogEditorMenus, Warning, TEXT("Unknown string command handler type: '%s'"), *StringCommand.TypeName.ToString());
+			UE_LOG(LogEditorMenus, Warning, TEXT("Unknown string command handler type: '%s'"), *TypeName.ToString());
 		}
 	}
 }
 
-FUIAction UEditorMenuSubsystem::ConvertUIAction(const FEditorMenuStringCommand& StringCommand, const FEditorMenuContext& Context)
-{
-	FUIAction UIAction;
-
-	if (StringCommand.String.Len() > 0)
-	{
-		UIAction.ExecuteAction.BindUObject(this, &UEditorMenuSubsystem::ExecuteStringCommand, StringCommand, Context);
-	}
-
-	return UIAction;
-}
-
 UEditorMenu* UEditorMenuSubsystem::FindSubMenuToGenerateWith(const FName InParentName, const FName InChildName)
 {
 	FName BaseName = InParentName;
diff --git a/Engine/Source/Editor/EditorMenus/Private/EditorMenusBlueprintLibrary.cpp b/Engine/Source/Editor/EditorMenus/Private/EditorMenusBlueprintLibrary.cpp
--- a/Engine/Source/Editor/EditorMenus/Private/EditorMenusBlueprintLibrary.cpp
+++ b/Engine/Source/Editor/EditorMenus/Private/EditorMenusBlueprintLibrary.cpp
@@ -56,8 +56,8 @@ void UEditorMenuEntryExtensions::SetIcon(UPARAM(ref) FEditorMenuEntry& Target, c
 void UEditorMenuEntryExtensions::SetStringCommand(UPARAM(ref) FEditorMenuEntry& Target, const FName TypeName, const FString& String)
 {
 	Target.ResetActions();
-	Target.StringCommand.TypeName = TypeName;
-	Target.StringCommand.String = String;
+	Target.StringExecuteAction.TypeName = TypeName;
+	Target.StringExecuteAction.String = String;
 }
 
 FEditorMenuEntry UEditorMenuEntryExtensions::InitMenuEntry(const FName InOwner, const FName InName, const FText& InLabel, const FText& InToolTip, const FEditorMenuStringCommand& StringCommand)
@@ -65,7 +65,7 @@ FEditorMenuEntry UEditorMenuEntryExtensions::InitMenuEntry(const FName InOwner,
 	FEditorMenuEntry Entry(InOwner, InName, EMultiBlockType::MenuEntry);
 	Entry.Label = InLabel;
 	Entry.ToolTip = InToolTip;
-	Entry.StringCommand = StringCommand;
+	Entry.StringExecuteAction = StringCommand;
 	return Entry;
 }
 
diff --git a/Engine/Source/Editor/EditorMenus/Public/EditorMenu.h b/Engine/Source/Editor/EditorMenus/Public/EditorMenu.h
--- a/Engine/Source/Editor/EditorMenus/Public/EditorMenu.h
+++ b/Engine/Source/Editor/EditorMenus/Public/EditorMenu.h
@@ -60,6 +60,10 @@ public:
 		return Context.Find<TContextType>();
 	}
 
+	//~ Begin UObject Interface
+	virtual bool IsDestructionThreadSafe() const { return false; }
+	//~ End UObject Interface
+
 	friend class UEditorMenuSubsystem;
 
 private:
diff --git a/Engine/Source/Editor/EditorMenus/Public/EditorMenuDelegates.h b/Engine/Source/Editor/EditorMenus/Public/EditorMenuDelegates.h
--- a/Engine/Source/Editor/EditorMenus/Public/EditorMenuDelegates.h
+++ b/Engine/Source/Editor/EditorMenus/Public/EditorMenuDelegates.h
@@ -5,6 +5,8 @@
 #include "CoreMinimal.h"
 
 #include "EditorMenuContext.h"
+#include "EditorMenuMisc.h"
+
 #include "Framework/Commands/UIAction.h"
 #include "Framework/Commands/UICommandList.h"
 #include "Framework/MultiBox/MultiBoxBuilder.h"
@@ -20,6 +22,7 @@ DECLARE_DELEGATE_OneParam(FNewEditorMenuSectionDelegate, FEditorMenuSection&);
 DECLARE_DELEGATE_OneParam(FNewEditorMenuDelegate, UEditorMenu*);
 DECLARE_DELEGATE_TwoParams(FNewEditorMenuDelegateLegacy, class FMenuBuilder&, UEditorMenu*);
 DECLARE_DELEGATE_TwoParams(FNewToolBarDelegateLegacy, class FToolBarBuilder&, UEditorMenu*);
+DECLARE_DELEGATE_RetVal_OneParam(TSharedRef<SWidget>, FNewEditorMenuWidget, const FEditorMenuContext&);
 
 DECLARE_DELEGATE_OneParam(FEditorMenuExecuteAction, const FEditorMenuContext&);
 DECLARE_DELEGATE_RetVal_OneParam(bool, FEditorMenuCanExecuteAction, const FEditorMenuContext&);
@@ -49,7 +52,7 @@ struct EDITORMENUS_API FEditorUIAction
 USTRUCT(BlueprintType)
 struct EDITORMENUS_API FEditorDynamicUIAction
 {
-	GENERATED_USTRUCT_BODY()
+	GENERATED_BODY()
 
 	FEditorDynamicUIAction() {}
 
@@ -66,36 +69,49 @@ struct EDITORMENUS_API FEditorDynamicUIAction
 	FEditorMenuDynamicIsActionButtonVisible IsActionVisibleDelegate;
 };
 
+struct EDITORMENUS_API FNewEditorMenuWidgetChoice
+{
+public:
+	FNewEditorMenuWidgetChoice() {}
+	FNewEditorMenuWidgetChoice(const FOnGetContent& InOnGetContent) : OnGetContent(InOnGetContent) {}
+	FNewEditorMenuWidgetChoice(const FNewEditorMenuWidget& InNewEditorMenuWidget) : NewEditorMenuWidget(InNewEditorMenuWidget) {}
+	FNewEditorMenuWidgetChoice(const FNewEditorMenuDelegate& InNewEditorMenu) : NewEditorMenu(InNewEditorMenu) {}
+
+	FOnGetContent OnGetContent;
+	FNewEditorMenuWidget NewEditorMenuWidget;
+	FNewEditorMenuDelegate NewEditorMenu;
+};
+
 struct EDITORMENUS_API FEditorUIActionChoice
 {
 public:
 	FEditorUIActionChoice() {}
-	FEditorUIActionChoice(const FUIAction& InAction) : Action(InAction) { }
-	FEditorUIActionChoice(const FExecuteAction& InExecuteAction) : Action(InExecuteAction) { }
-	FEditorUIActionChoice(const FEditorUIAction& InAction) : EditorAction(InAction) { }
-	FEditorUIActionChoice(const FEditorDynamicUIAction& InAction) : DynamicEditorAction(InAction) { }
+	FEditorUIActionChoice(const FUIAction& InAction) : Action(InAction) {}
+	FEditorUIActionChoice(const FExecuteAction& InExecuteAction) : Action(InExecuteAction) {}
+	FEditorUIActionChoice(const FEditorUIAction& InAction) : EditorAction(InAction) {}
+	FEditorUIActionChoice(const FEditorDynamicUIAction& InAction) : DynamicEditorAction(InAction) {}
 	FEditorUIActionChoice(const FEditorMenuExecuteAction& InExecuteAction) : EditorAction(InExecuteAction) {}
 	FEditorUIActionChoice(const TSharedPtr< const FUICommandInfo >& InCommand, const FUICommandList& InCommandList);
 
-	FUIAction* GetUIAction()
+	const FUIAction* GetUIAction() const
 	{
 		return Action.IsSet() ? &Action.GetValue() : nullptr;
 	}
 
-	FEditorUIAction* GetEditorUIAction()
+	const FEditorUIAction* GetEditorUIAction() const
 	{
 		return EditorAction.IsSet() ? &EditorAction.GetValue() : nullptr;
 	}
 
-	FEditorDynamicUIAction* GetEditorDynamicUIAction()
+	const FEditorDynamicUIAction* GetEditorDynamicUIAction() const
 	{
 		return DynamicEditorAction.IsSet() ? &DynamicEditorAction.GetValue() : nullptr;
 	}
 
 private:
-	TOptional<FUIAction> Action; // 72 bytes
-	TOptional<FEditorUIAction> EditorAction; // 64 bytes
-	TOptional<FEditorDynamicUIAction> DynamicEditorAction; // 80 bytes
+	TOptional<FUIAction> Action;
+	TOptional<FEditorUIAction> EditorAction;
+	TOptional<FEditorDynamicUIAction> DynamicEditorAction;
 };
 
 struct EDITORMENUS_API FNewEditorMenuChoice
diff --git a/Engine/Source/Editor/EditorMenus/Public/EditorMenuEntry.h b/Engine/Source/Editor/EditorMenus/Public/EditorMenuEntry.h
--- a/Engine/Source/Editor/EditorMenus/Public/EditorMenuEntry.h
+++ b/Engine/Source/Editor/EditorMenus/Public/EditorMenuEntry.h
@@ -20,7 +20,7 @@ class UEditorMenuEntryScript;
 USTRUCT(BlueprintType)
 struct EDITORMENUS_API FSubMenuEntryData
 {
-	GENERATED_USTRUCT_BODY()
+	GENERATED_BODY()
 
 	FSubMenuEntryData() :
 		bIsSubMenu(false),
@@ -47,44 +47,33 @@ struct EDITORMENUS_API FSubMenuEntryData
 USTRUCT(BlueprintType)
 struct EDITORMENUS_API FToolBarData
 {
-	GENERATED_USTRUCT_BODY()
+	GENERATED_BODY()
 
 	FToolBarData() : bSimpleComboBox(false) {}
 
 	UPROPERTY(EditAnywhere, BlueprintReadWrite, Category = "Editor UI")
 	bool bSimpleComboBox;
 
-	FOnGetContent ComboButtonContextMenuGenerator;
+	FNewEditorMenuWidgetChoice ComboButtonContextMenuGenerator;
 	FNewToolBarDelegateLegacy ConstructLegacy;
 };
 
-USTRUCT(BlueprintType)
-struct EDITORMENUS_API FEditorMenuStringCommand
-{
-	GENERATED_USTRUCT_BODY()
-
-	UPROPERTY(EditAnywhere, BlueprintReadWrite, Category = "Editor UI")
-	FName TypeName;
-
-	UPROPERTY(EditAnywhere, BlueprintReadWrite, Category = "Editor UI")
-	FString String;
-};
-
 USTRUCT(BlueprintType)
 struct EDITORMENUS_API FEditorMenuEntry
 {
-	GENERATED_USTRUCT_BODY()
+	GENERATED_BODY()
 
 	FEditorMenuEntry();
 	FEditorMenuEntry(const FEditorMenuOwner InOwner, const FName InName, EMultiBlockType InType);
 
 	static FEditorMenuEntry InitMenuEntry(const FName InName, const TAttribute<FText>& InLabel, const TAttribute<FText>& InToolTip, const TAttribute<FSlateIcon>& InIcon, const FEditorUIActionChoice& InAction, const EUserInterfaceActionType UserInterfaceActionType = EUserInterfaceActionType::Button, const FName InTutorialHighlightName = NAME_None);
 	static FEditorMenuEntry InitMenuEntry(const TSharedPtr< const FUICommandInfo >& InCommand, const TAttribute<FText>& InLabelOverride = TAttribute<FText>(), const TAttribute<FText>& InToolTipOverride = TAttribute<FText>(), const TAttribute<FSlateIcon>& InIconOverride = TAttribute<FSlateIcon>(), const FName InTutorialHighlightName = NAME_None, const FName InNameOverride = NAME_None);
+	static FEditorMenuEntry InitMenuEntry(const FName InName, const FEditorUIActionChoice& InAction, const TSharedRef<SWidget>& Widget);
 	static FEditorMenuEntry InitSubMenu(const FName InParentMenu, const FName InName, const TAttribute<FText>& InLabel, const TAttribute<FText>& InToolTip, const FNewEditorMenuChoice& InMakeMenu, bool bInOpenSubMenuOnClick = false, const TAttribute<FSlateIcon>& InIcon = TAttribute<FSlateIcon>(), const bool ShouldCloseWindowAfterMenuSelection = true);
 
 	static FEditorMenuEntry InitToolBarButton(const FName InName, const FEditorUIActionChoice& InAction, const TAttribute<FText>& InLabel = TAttribute<FText>(), const TAttribute<FText>& InToolTip = TAttribute<FText>(), const TAttribute<FSlateIcon>& InIcon = TAttribute<FSlateIcon>(), const EUserInterfaceActionType UserInterfaceActionType = EUserInterfaceActionType::Button, FName InTutorialHighlightName = NAME_None);
 	static FEditorMenuEntry InitToolBarButton(const TSharedPtr< const FUICommandInfo >& InCommand, const TAttribute<FText>& InLabelOverride = TAttribute<FText>(), const TAttribute<FText>& InToolTipOverride = TAttribute<FText>(), const TAttribute<FSlateIcon>& InIconOverride = TAttribute<FSlateIcon>(), FName InTutorialHighlightName = NAME_None, const FName InNameOverride = NAME_None);
-	static FEditorMenuEntry InitComboButton(const FName InName, const FEditorUIActionChoice& InAction, const FOnGetContent& InMenuContentGenerator, const TAttribute<FText>& InLabelOverride = TAttribute<FText>(), const TAttribute<FText>& InToolTipOverride = TAttribute<FText>(), const TAttribute<FSlateIcon>& InIconOverride = TAttribute<FSlateIcon>(), bool bInSimpleComboBox = false, FName InTutorialHighlightName = NAME_None);
+	static FEditorMenuEntry InitComboButton(const FName InName, const FEditorUIActionChoice& InAction, const FNewEditorMenuWidgetChoice& InMenuContentGenerator, const TAttribute<FText>& InLabelOverride = TAttribute<FText>(), const TAttribute<FText>& InToolTipOverride = TAttribute<FText>(), const TAttribute<FSlateIcon>& InIconOverride = TAttribute<FSlateIcon>(), bool bInSimpleComboBox = false, FName InTutorialHighlightName = NAME_None);
 
 	static FEditorMenuEntry InitMenuSeparator(const FName InName);
 	static FEditorMenuEntry InitToolBarSeparator(const FName InName);
@@ -132,6 +121,9 @@ public:
 	UPROPERTY(EditAnywhere, BlueprintReadWrite, Category = "Editor UI")
 	UEditorMenuEntryScript* ScriptObject;
 
+	/** Optional delegate that returns a widget to use as this menu entry */
+	FNewEditorMenuWidget MakeWidget;
+
 private:
 
 	friend class UEditorMenuSubsystem;
@@ -143,7 +135,7 @@ private:
 
 	FEditorUIActionChoice Action;
 
-	FEditorMenuStringCommand StringCommand;
+	FEditorMenuStringCommand StringExecuteAction;
 
 	TSharedPtr< const FUICommandInfo > Command;
 	TSharedPtr< const FUICommandList > CommandList;
diff --git a/Engine/Source/Editor/EditorMenus/Public/EditorMenuMisc.h b/Engine/Source/Editor/EditorMenus/Public/EditorMenuMisc.h
--- a/Engine/Source/Editor/EditorMenus/Public/EditorMenuMisc.h
+++ b/Engine/Source/Editor/EditorMenus/Public/EditorMenuMisc.h
@@ -3,9 +3,45 @@
 #pragma once
 
 #include "CoreMinimal.h"
+#include "Framework/Commands/UIAction.h"
 
 #include "EditorMenuMisc.generated.h"
 
+struct FEditorMenuContext;
+
+UENUM(BlueprintType)
+enum class EEditorMenuStringCommandType : uint8
+{
+	Command,
+	Python,
+	Other
+};
+
+USTRUCT(BlueprintType)
+struct EDITORMENUS_API FEditorMenuStringCommand
+{
+	GENERATED_BODY()
+
+	UPROPERTY(EditAnywhere, BlueprintReadWrite, Category = "Editor UI")
+	EEditorMenuStringCommandType Type;
+
+	UPROPERTY(EditAnywhere, BlueprintReadWrite, Category = "Editor UI")
+	FName TypeName;
+
+	UPROPERTY(EditAnywhere, BlueprintReadWrite, Category = "Editor UI")
+	FString String;
+
+private:
+
+	friend class UEditorMenuSubsystem;
+
+	bool IsBound() const { return String.Len() > 0; }
+
+	FExecuteAction ToExecuteAction(const FEditorMenuContext& Context) const;
+
+	FName GetTypeName() const;
+};
+
 UENUM(BlueprintType)
 enum class EEditorMenuInsertType : uint8
 {
@@ -18,7 +54,7 @@ enum class EEditorMenuInsertType : uint8
 USTRUCT(BlueprintType)
 struct FEditorMenuInsert
 {
-	GENERATED_USTRUCT_BODY()
+	GENERATED_BODY()
 
 	FEditorMenuInsert() : Position(EEditorMenuInsertType::Default) {}
 	FEditorMenuInsert(FName InName, EEditorMenuInsertType InPosition) : Name(InName), Position(InPosition) {}
@@ -49,3 +85,4 @@ struct FEditorMenuInsert
 		return Position == EEditorMenuInsertType::Before || Position == EEditorMenuInsertType::After;
 	}
 };
+
diff --git a/Engine/Source/Editor/EditorMenus/Public/EditorMenuOwner.h b/Engine/Source/Editor/EditorMenus/Public/EditorMenuOwner.h
--- a/Engine/Source/Editor/EditorMenus/Public/EditorMenuOwner.h
+++ b/Engine/Source/Editor/EditorMenus/Public/EditorMenuOwner.h
@@ -10,7 +10,7 @@
 USTRUCT(BlueprintType, meta=(HasNativeBreak="EditorMenus.EditorMenuEntryExtensions.BreakEditorMenuOwner", HasNativeMake="EditorMenus.EditorMenuEntryExtensions.MakeEditorMenuOwner"))
 struct EDITORMENUS_API FEditorMenuOwner
 {
-	GENERATED_USTRUCT_BODY()
+	GENERATED_BODY()
 
 private:
 	struct FStoreName
diff --git a/Engine/Source/Editor/EditorMenus/Public/EditorMenuSection.h b/Engine/Source/Editor/EditorMenus/Public/EditorMenuSection.h
--- a/Engine/Source/Editor/EditorMenus/Public/EditorMenuSection.h
+++ b/Engine/Source/Editor/EditorMenus/Public/EditorMenuSection.h
@@ -34,6 +34,7 @@ public:
 
 	FEditorMenuEntry& AddEntryObject(UEditorMenuEntryScript* InObject);
 	FEditorMenuEntry& AddEntry(const FEditorMenuEntry& Args);
+	FEditorMenuEntry& AddMenuEntry(const FName InName, const TAttribute<FText>& InLabel, const TAttribute<FText>& InToolTip, const TAttribute<FSlateIcon>& InIcon, const FEditorUIActionChoice& InAction, const EUserInterfaceActionType UserInterfaceActionType = EUserInterfaceActionType::Button, const FName InTutorialHighlightName = NAME_None);
 	FEditorMenuEntry& AddMenuEntry(const TSharedPtr< const FUICommandInfo >& InCommand, const TAttribute<FText>& InLabelOverride = TAttribute<FText>(), const TAttribute<FText>& InToolTipOverride = TAttribute<FText>(), const TAttribute<FSlateIcon>& InIconOverride = TAttribute<FSlateIcon>(), const FName InTutorialHighlightName = NAME_None, const FName InNameOverride = NAME_None);
 
 	FEditorMenuEntry& AddDynamicEntry(const FName InName, const FNewEditorMenuSectionDelegate& InConstruct);
diff --git a/Engine/Source/Editor/EditorMenus/Public/EditorMenuSubsystem.h b/Engine/Source/Editor/EditorMenus/Public/EditorMenuSubsystem.h
--- a/Engine/Source/Editor/EditorMenus/Public/EditorMenuSubsystem.h
+++ b/Engine/Source/Editor/EditorMenus/Public/EditorMenuSubsystem.h
@@ -8,11 +8,6 @@
 #include "EditorSubsystem.h"
 #include "Templates/SubclassOf.h"
 #include "Textures/SlateIcon.h"
-#include "Framework/Commands/UIAction.h"
-#include "Framework/MultiBox/MultiBox.h"
-#include "Framework/MultiBox/MultiBoxBuilder.h"
-#include "Framework/MultiBox/MultiBoxExtender.h"
-#include "Framework/MultiBox/MultiBoxDefs.h"
 #include "Misc/Attribute.h"
 #include "Editor.h"
 
@@ -29,7 +24,6 @@
 #include "EditorMenuSubsystem.generated.h"
 
 class FMenuBuilder;
-class FUICommandInfo;
 class FMultiBox;
 class SWidget;
 
@@ -97,16 +91,10 @@ public:
 
 	UEditorMenuSubsystem();
 
-	/** Get EditorMenuSubsystem if GEditor is valid and cause EditorMenus module to load. */
 	static inline UEditorMenuSubsystem* Get()
 	{
-		if (GEditor)
-		{
-			FModuleManager::LoadModuleChecked<IEditorMenusModule>("EditorMenus");
-			return GEditor->GetEditorSubsystem<UEditorMenuSubsystem>();
-		}
-
-		return nullptr;
+		FModuleManager::LoadModuleChecked<IEditorMenusModule>("EditorMenus");
+		return GEditor->GetEditorSubsystem<UEditorMenuSubsystem>();
 	}
 
 	/** Try to get EditorMenuSubsystem if GEditor is valid without forcing EditorMenus module to load. */
@@ -170,6 +158,14 @@ public:
 	UFUNCTION(BlueprintCallable, Category = "Editor UI")
 	UEditorMenu* FindMenu(const FName Name);
 
+	/**
+	 * Determines if a menu has already been registered.
+	 * @param	Name	Name of the menu to find.
+	 * @return	bool	True if menu has already been registered.
+	 */
+	UFUNCTION(BlueprintCallable, Category = "Editor UI")
+	bool IsMenuRegistered(const FName Name) const;
+
 	/** Rebuilds all widgets generated from a specific menu. */
 	UFUNCTION(BlueprintCallable, Category = "Editor UI")
 	bool RefreshMenuWidget(const FName Name);
@@ -256,6 +252,7 @@ public:
 	void UnregisterStringCommandHandler(const FName InName);
 
 	friend struct FEditorMenuOwnerScoped;
+	friend struct FEditorMenuStringCommand;
 
 	//~ Begin UObject Interface
 	static void AddReferencedObjects(UObject* InThis, FReferenceCollector& Collector);
@@ -296,14 +293,16 @@ private:
 
 	TSharedRef<SWidget> GenerateToolbarComboButtonMenu(const FName SubMenuFullName, FEditorMenuContext InContext);
 
+	FOnGetContent ConvertWidgetChoice(const FNewEditorMenuWidgetChoice& Choice, const FEditorMenuContext& Context) const;
+
 	/** Converts a string command to a FUIAction */
-	FUIAction ConvertUIAction(const FEditorMenuStringCommand& StringCommand, const FEditorMenuContext& Context);
-	static FUIAction ConvertUIAction(FEditorUIActionChoice& Actions, FEditorMenuContext& Context);
-	static FUIAction ConvertUIAction(FEditorUIAction& Actions, FEditorMenuContext& Context);
-	static FUIAction ConvertUIAction(FEditorDynamicUIAction& Actions, FEditorMenuContext& Context);
-	static FUIAction ConvertUIAction(UEditorMenuEntryScript* ScriptObject, FEditorMenuContext& Context);
+	static FUIAction ConvertUIAction(const FEditorMenuEntry& Block, const FEditorMenuContext& Context);
+	static FUIAction ConvertUIAction(const FEditorUIActionChoice& Choice, const FEditorMenuContext& Context);
+	static FUIAction ConvertUIAction(const FEditorUIAction& Actions, const FEditorMenuContext& Context);
+	static FUIAction ConvertUIAction(const FEditorDynamicUIAction& Actions, const FEditorMenuContext& Context);
+	static FUIAction ConvertScriptObjectToUIAction(UEditorMenuEntryScript* ScriptObject, const FEditorMenuContext& Context);
 
-	void ExecuteStringCommand(const FEditorMenuStringCommand StringCommand, const FEditorMenuContext Context);
+	static void ExecuteStringCommand(const FEditorMenuStringCommand StringCommand, const FEditorMenuContext Context);
 
 	void FillMenuDynamic(FMenuBuilder& Builder, FNewEditorMenuDelegate InConstructMenu);
 
