Editor Menu widget support updates

diff --git a/Engine/Source/Editor/EditorMenus/Private/EditorMenu.cpp b/Engine/Source/Editor/EditorMenus/Private/EditorMenu.cpp
--- a/Engine/Source/Editor/EditorMenus/Private/EditorMenu.cpp
+++ b/Engine/Source/Editor/EditorMenus/Private/EditorMenu.cpp
@@ -159,9 +159,8 @@ void UEditorMenu::AddMenuEntryObject(UEditorMenuEntryScript* InObject)
 
 UEditorMenu* UEditorMenu::AddSubMenu(const FEditorMenuOwner InOwner, const FName SectionName, const FName InName, const FText& InLabel, const FText& InToolTip)
 {
-	FEditorMenuEntry Args = FEditorMenuEntry::InitMenuEntry(InName, InLabel, InToolTip, FSlateIcon(), FUIAction());
+	FEditorMenuEntry Args = FEditorMenuEntry::InitSubMenu(MenuName, InName, InLabel, InToolTip, FNewEditorMenuChoice());
 	Args.Owner = InOwner;
-	Args.SubMenu.bIsSubMenu = true;
 	FindOrAddSection(NAME_None).AddEntry(Args);
 	return UEditorMenuSubsystem::Get()->ExtendMenu(*(MenuName.ToString() + TEXT(".") + InName.ToString()));
 }
diff --git a/Engine/Source/Editor/EditorMenus/Private/EditorMenuContext.cpp b/Engine/Source/Editor/EditorMenus/Private/EditorMenuContext.cpp
--- a/Engine/Source/Editor/EditorMenus/Private/EditorMenuContext.cpp
+++ b/Engine/Source/Editor/EditorMenus/Private/EditorMenuContext.cpp
@@ -87,6 +87,14 @@ TSharedPtr<FExtender> FEditorMenuContext::GetAllExtenders()
 	return TSharedPtr<FExtender>();
 }
 
+void FEditorMenuContext::ResetExtenders()
+{
+	if (ExtensibilityManager.IsValid())
+	{
+		ExtensibilityManager.Reset();
+	}
+}
+
 void FEditorMenuContext::AppendObjects(const TArray<UObject*>& InObjects)
 {
 	for (UObject* Object : InObjects)
diff --git a/Engine/Source/Editor/EditorMenus/Private/EditorMenuEntry.cpp b/Engine/Source/Editor/EditorMenus/Private/EditorMenuEntry.cpp
--- a/Engine/Source/Editor/EditorMenus/Private/EditorMenuEntry.cpp
+++ b/Engine/Source/Editor/EditorMenus/Private/EditorMenuEntry.cpp
@@ -26,7 +26,7 @@ FEditorMenuEntry::FEditorMenuEntry(const FEditorMenuOwner InOwner, const FName I
 	Owner(InOwner),
 	Type(InType),
 	UserInterfaceActionType(EUserInterfaceActionType::Button),
-	bShouldCloseWindowAfterMenuSelection(true)	,
+	bShouldCloseWindowAfterMenuSelection(true),
 	ScriptObject(nullptr)
 {
 }
@@ -74,9 +74,10 @@ FEditorMenuEntry FEditorMenuEntry::InitSubMenu(const FName InParentMenu, const F
 	Entry.Label = InLabel;
 	Entry.ToolTip = InToolTip;
 	Entry.Icon = InIcon;
-
-	Entry.SubMenu = FSubMenuEntryData(InMakeMenu, bInOpenSubMenuOnClick);
 	Entry.bShouldCloseWindowAfterMenuSelection = bInShouldCloseWindowAfterMenuSelection;
+	Entry.SubMenuData.bIsSubMenu = true;
+	Entry.SubMenuData.ConstructMenu = InMakeMenu;
+	Entry.SubMenuData.bOpenSubMenuOnClick = bInOpenSubMenuOnClick;
 	return Entry;
 }
 
@@ -108,9 +109,8 @@ FEditorMenuEntry FEditorMenuEntry::InitComboButton(const FName InName, const FEd
 	Entry.ToolTip = InToolTip;
 	Entry.Icon = InIcon;
 	Entry.Action = InAction;
-
-	Entry.ToolBar.ComboButtonContextMenuGenerator = InMenuContentGenerator;
-	Entry.ToolBar.bSimpleComboBox = bInSimpleComboBox;
+	Entry.ToolBarData.ComboButtonContextMenuGenerator = InMenuContentGenerator;
+	Entry.ToolBarData.bSimpleComboBox = bInSimpleComboBox;
 	return Entry;
 }
 
@@ -124,6 +124,16 @@ FEditorMenuEntry FEditorMenuEntry::InitToolBarSeparator(const FName InName)
 	return FEditorMenuEntry(UEditorMenuSubsystem::Get()->CurrentOwner(), InName, EMultiBlockType::ToolBarSeparator);
 }
 
+FEditorMenuEntry FEditorMenuEntry::InitWidget(const FName InName, const TSharedRef<SWidget>& Widget, const FText& Label, bool bNoIndent, bool bSearchable)
+{
+	FEditorMenuEntry Entry(UEditorMenuSubsystem::Get()->CurrentOwner(), InName, EMultiBlockType::Widget);
+	Entry.Label = Label;
+	Entry.MakeWidget.BindLambda([Widget](const FEditorMenuContext&) { return Widget; });
+	Entry.WidgetData.bNoIndent = bNoIndent;
+	Entry.WidgetData.bSearchable = bSearchable;
+	return Entry;
+}
+
 void FEditorMenuEntry::ResetActions()
 {
 	Action = FEditorUIActionChoice();
@@ -133,3 +143,9 @@ void FEditorMenuEntry::ResetActions()
 	// Note: Cannot reset ScriptObject as it would also remove label and other data
 	//ScriptObject = nullptr;
 }
+
+bool FEditorMenuEntry::IsScriptObjectDynamicConstruct() const
+{
+	static const FName ConstructMenuEntryName = GET_FUNCTION_NAME_CHECKED(UEditorMenuEntryScript, ConstructMenuEntry);
+	return ScriptObject && ScriptObject->GetClass()->IsFunctionImplementedInScript(ConstructMenuEntryName);
+}
diff --git a/Engine/Source/Editor/EditorMenus/Private/EditorMenuEntryScript.cpp b/Engine/Source/Editor/EditorMenus/Private/EditorMenuEntryScript.cpp
--- a/Engine/Source/Editor/EditorMenus/Private/EditorMenuEntryScript.cpp
+++ b/Engine/Source/Editor/EditorMenus/Private/EditorMenuEntryScript.cpp
@@ -128,12 +128,6 @@ void UEditorMenuEntryScript::ToMenuEntry(FEditorMenuEntry& Output)
 	Output.Owner = Data.OwnerName;
 }
 
-bool FEditorMenuEntry::IsScriptObjectDynamicConstruct() const
-{
-	static const FName ConstructMenuEntryName = GET_FUNCTION_NAME_CHECKED(UEditorMenuEntryScript, ConstructMenuEntry);
-	return ScriptObject && ScriptObject->GetClass()->IsFunctionImplementedInScript(ConstructMenuEntryName);
-}
-
 FEditorMenuEntryScriptDataAdvanced::FEditorMenuEntryScriptDataAdvanced() :
 	EntryType(EMultiBlockType::MenuEntry),
 	UserInterfaceActionType(EUserInterfaceActionType::Button),
diff --git a/Engine/Source/Editor/EditorMenus/Private/EditorMenuSubsystem.cpp b/Engine/Source/Editor/EditorMenus/Private/EditorMenuSubsystem.cpp
--- a/Engine/Source/Editor/EditorMenus/Private/EditorMenuSubsystem.cpp
+++ b/Engine/Source/Editor/EditorMenus/Private/EditorMenuSubsystem.cpp
@@ -515,28 +515,28 @@ void UEditorMenuSubsystem::PopulateMenuBuilder(FMenuBuilder& MenuBuilder, UEdito
 
 			if (Block.Type == EMultiBlockType::MenuEntry)
 			{
-				if (Block.SubMenu.bIsSubMenu)
+				if (Block.IsSubMenu())
 				{
-					if (Block.SubMenu.Construct.NewMenuDelegate.IsBound())
+					if (Block.SubMenuData.ConstructMenu.NewMenuDelegate.IsBound())
 					{
 						MenuBuilder.AddSubMenu(
 							Block.Label,
 							Block.ToolTip,
-							Block.SubMenu.Construct.NewMenuDelegate,
-							Block.SubMenu.bOpenSubMenuOnClick,
+							Block.SubMenuData.ConstructMenu.NewMenuDelegate,
+							Block.SubMenuData.bOpenSubMenuOnClick,
 							Block.Icon.Get(),
 							Block.bShouldCloseWindowAfterMenuSelection,
 							Block.Name
 						);
 					}
-					else if (Block.SubMenu.Construct.NewEditorMenuDelegate.IsBound())
+					else if (Block.SubMenuData.ConstructMenu.NewEditorMenuDelegate.IsBound())
 					{
 						// SubMenu constructed each time it is opened
 						MenuBuilder.AddSubMenu(
 							Block.Label,
 							Block.ToolTip,
-							FNewMenuDelegate::CreateUObject(this, &UEditorMenuSubsystem::FillMenuDynamic, Block.SubMenu.Construct.NewEditorMenuDelegate),
-							Block.SubMenu.bOpenSubMenuOnClick,
+							FNewMenuDelegate::CreateUObject(this, &UEditorMenuSubsystem::FillMenuDynamic, Block.SubMenuData.ConstructMenu.NewEditorMenuDelegate),
+							Block.SubMenuData.bOpenSubMenuOnClick,
 							Block.Icon.Get(),
 							Block.bShouldCloseWindowAfterMenuSelection,
 							Block.Name
@@ -557,7 +557,7 @@ void UEditorMenuSubsystem::PopulateMenuBuilder(FMenuBuilder& MenuBuilder, UEdito
 							}
 							else
 							{
-								MenuBuilder.AddSubMenu(Widget.ToSharedRef(), NewMenuDelegate, Block.SubMenu.bOpenSubMenuOnClick, Block.bShouldCloseWindowAfterMenuSelection);
+								MenuBuilder.AddSubMenu(Widget.ToSharedRef(), NewMenuDelegate, Block.SubMenuData.bOpenSubMenuOnClick, Block.bShouldCloseWindowAfterMenuSelection);
 							}
 						}
 						else
@@ -566,7 +566,7 @@ void UEditorMenuSubsystem::PopulateMenuBuilder(FMenuBuilder& MenuBuilder, UEdito
 								Block.Label,
 								Block.ToolTip,
 								NewMenuDelegate,
-								Block.SubMenu.bOpenSubMenuOnClick,
+								Block.SubMenuData.bOpenSubMenuOnClick,
 								Block.Icon.Get(),
 								Block.bShouldCloseWindowAfterMenuSelection,
 								Block.Name
@@ -602,6 +602,14 @@ void UEditorMenuSubsystem::PopulateMenuBuilder(FMenuBuilder& MenuBuilder, UEdito
 			{
 				MenuBuilder.AddMenuSeparator(Block.Name);
 			}
+			else if (Block.Type == EMultiBlockType::Widget)
+			{
+				MenuBuilder.AddWidget(Widget.ToSharedRef(), Block.Label.Get(), Block.WidgetData.bNoIndent, Block.WidgetData.bSearchable);
+			}
+			else
+			{
+				UE_LOG(LogEditorMenus, Warning, TEXT("Menu '%s', item '%s', type not currently supported: %d"), *MenuData->MenuName.ToString(), *Block.Name.ToString(), Block.Type);
+			}
 		}
 
 		MenuBuilder.EndSection();
@@ -624,13 +632,20 @@ void UEditorMenuSubsystem::PopulateToolBarBuilder(FToolBarBuilder& ToolBarBuilde
 
 		for (FEditorMenuEntry& Block : Section.Blocks)
 		{
-			if (Block.ToolBar.ConstructLegacy.IsBound())
+			if (Block.ToolBarData.ConstructLegacy.IsBound())
 			{
-				Block.ToolBar.ConstructLegacy.Execute(ToolBarBuilder, MenuData);
+				Block.ToolBarData.ConstructLegacy.Execute(ToolBarBuilder, MenuData);
 				continue;
 			}
 
 			FUIAction UIAction = ConvertUIAction(Block, MenuData->Context);
+
+			TSharedPtr<SWidget> Widget;
+			if (Block.MakeWidget.IsBound())
+			{
+				Widget = Block.MakeWidget.Execute(MenuData->Context);
+			}
+
 			if (Block.Type == EMultiBlockType::ToolBarButton)
 			{
 				if (Block.Command.IsValid())
@@ -649,25 +664,29 @@ void UEditorMenuSubsystem::PopulateToolBarBuilder(FToolBarBuilder& ToolBarBuilde
 			}
 			else if (Block.Type == EMultiBlockType::ToolBarComboButton)
 			{
-				FOnGetContent OnGetContent = ConvertWidgetChoice(Block.ToolBar.ComboButtonContextMenuGenerator, MenuData->Context);
+				FOnGetContent OnGetContent = ConvertWidgetChoice(Block.ToolBarData.ComboButtonContextMenuGenerator, MenuData->Context);
 				if (OnGetContent.IsBound())
 				{
-					ToolBarBuilder.AddComboButton(UIAction, OnGetContent, Block.Label, Block.ToolTip, Block.Icon, Block.ToolBar.bSimpleComboBox, Block.TutorialHighlightName);
+					ToolBarBuilder.AddComboButton(UIAction, OnGetContent, Block.Label, Block.ToolTip, Block.Icon, Block.ToolBarData.bSimpleComboBox, Block.TutorialHighlightName);
 				}
 				else
 				{
 					FName SubMenuFullName = JoinMenuPaths(MenuData->MenuName, Block.Name);
 					FOnGetContent Delegate = FOnGetContent::CreateUObject(this, &UEditorMenuSubsystem::GenerateToolbarComboButtonMenu, SubMenuFullName, MenuData->Context);
-					ToolBarBuilder.AddComboButton(UIAction, Delegate, Block.Label, Block.ToolTip, Block.Icon, Block.ToolBar.bSimpleComboBox, Block.TutorialHighlightName);
+					ToolBarBuilder.AddComboButton(UIAction, Delegate, Block.Label, Block.ToolTip, Block.Icon, Block.ToolBarData.bSimpleComboBox, Block.TutorialHighlightName);
 				}
 			}
 			else if (Block.Type == EMultiBlockType::ToolBarSeparator)
 			{
 				ToolBarBuilder.AddSeparator(Block.Name);
 			}
+			else if (Block.Type == EMultiBlockType::Widget)
+			{
+				ToolBarBuilder.AddWidget(Widget.ToSharedRef(), Block.TutorialHighlightName, Block.WidgetData.bSearchable);
+			}
 			else
 			{
-				UE_LOG(LogEditorMenus, Warning, TEXT("Toolbar '%s', item '%s', has unknown type: %d"), *MenuData->MenuName.ToString(), *Block.Name.ToString(), Block.Type);
+				UE_LOG(LogEditorMenus, Warning, TEXT("Toolbar '%s', item '%s', type not currently supported: %d"), *MenuData->MenuName.ToString(), *Block.Name.ToString(), Block.Type);
 			}
 		}
 
diff --git a/Engine/Source/Editor/EditorMenus/Public/EditorMenuContext.h b/Engine/Source/Editor/EditorMenus/Public/EditorMenuContext.h
--- a/Engine/Source/Editor/EditorMenus/Public/EditorMenuContext.h
+++ b/Engine/Source/Editor/EditorMenus/Public/EditorMenuContext.h
@@ -4,8 +4,6 @@
 
 #include "CoreMinimal.h"
 
-#include "Kismet/BlueprintFunctionLibrary.h"
-
 #include "EditorMenuContext.generated.h"
 
 class FUICommandList;
@@ -41,9 +39,9 @@ public:
 	template <typename TContextType>
 	TContextType* Find() const
 	{
-		for (const TWeakObjectPtr<UObject>& WeakObject : ContextObjects)
+		for (UObject* Object : ContextObjects)
 		{
-			if (TContextType* Result = Cast<TContextType>(WeakObject.Get()))
+			if (TContextType* Result = Cast<TContextType>(Object))
 			{
 				return Result;
 			}
@@ -57,6 +55,8 @@ public:
 	void AppendCommandList(TSharedPtr<FUICommandList> InCommandList);
 	void AddExtender(const TSharedPtr<FExtender>& InExtender);
 	TSharedPtr<FExtender> GetAllExtenders();
+	void ReplaceExtenders(const TSharedPtr<FExtender>& InExtender);
+	void ResetExtenders();
 
 	void AppendObjects(const TArray<UObject*>& InObjects);
 	void AddObject(UObject* InObject);
diff --git a/Engine/Source/Editor/EditorMenus/Public/EditorMenuEntry.h b/Engine/Source/Editor/EditorMenus/Public/EditorMenuEntry.h
--- a/Engine/Source/Editor/EditorMenus/Public/EditorMenuEntry.h
+++ b/Engine/Source/Editor/EditorMenus/Public/EditorMenuEntry.h
@@ -17,47 +17,65 @@
 
 class UEditorMenuEntryScript;
 
-USTRUCT(BlueprintType)
-struct EDITORMENUS_API FSubMenuEntryData
+struct FEditorMenuEntrySubMenuData
 {
-	GENERATED_BODY()
-
-	FSubMenuEntryData() :
+public:
+	FEditorMenuEntrySubMenuData() :
 		bIsSubMenu(false),
 		bOpenSubMenuOnClick(false)
 	{
 	}
 
-	FSubMenuEntryData(const FNewEditorMenuChoice& InConstruct, bool bInOpenSubMenuOnClick = false) :
-		bIsSubMenu(true),
-		bOpenSubMenuOnClick(bInOpenSubMenuOnClick),
-		Construct(InConstruct)
-	{
-	}
-
-	UPROPERTY(EditAnywhere, BlueprintReadWrite, Category = "Editor UI")
 	bool bIsSubMenu;
-
-	UPROPERTY(EditAnywhere, BlueprintReadWrite, Category = "Editor UI")
 	bool bOpenSubMenuOnClick;
-
-	FNewEditorMenuChoice Construct;
+	FNewEditorMenuChoice ConstructMenu;
 };
 
-USTRUCT(BlueprintType)
-struct EDITORMENUS_API FToolBarData
+struct FEditorMenuEntryToolBarData
 {
-	GENERATED_BODY()
+public:
+	FEditorMenuEntryToolBarData() :
+		bSimpleComboBox(false),
+		bIsFocusable(false),
+		bForceSmallIcons(false)
+	{
+	}
 
-	FToolBarData() : bSimpleComboBox(false) {}
+	TOptional< EVisibility > LabelVisibility;
 
-	UPROPERTY(EditAnywhere, BlueprintReadWrite, Category = "Editor UI")
+	/** If true, the icon and label won't be displayed */
 	bool bSimpleComboBox;
 
+	/** Whether ToolBar will have Focusable buttons */
+	bool bIsFocusable;
+
+	/** Whether this toolbar should always use small icons, regardless of the current settings */
+	bool bForceSmallIcons;
+
+	/** Delegate that generates a widget for this combo button's menu content.  Called when the menu is summoned. */
 	FNewEditorMenuWidgetChoice ComboButtonContextMenuGenerator;
+
+	/** Legacy delegate that generates a widget for this combo button's menu content.  Called when the menu is summoned. */
 	FNewToolBarDelegateLegacy ConstructLegacy;
 };
 
+
+struct FEditorMenuEntryWidgetData
+{
+public:
+	FEditorMenuEntryWidgetData() :
+		bNoIndent(false),
+		bSearchable(false)
+	{
+	}
+
+	/** Remove the padding from the left of the widget that lines it up with other menu items */
+	bool bNoIndent;
+
+	/** If true, widget will be searchable */
+	bool bSearchable;
+};
+
 USTRUCT(BlueprintType)
 struct EDITORMENUS_API FEditorMenuEntry
 {
@@ -78,6 +96,10 @@ struct EDITORMENUS_API FEditorMenuEntry
 	static FEditorMenuEntry InitMenuSeparator(const FName InName);
 	static FEditorMenuEntry InitToolBarSeparator(const FName InName);
 
+	static FEditorMenuEntry InitWidget(const FName InName, const TSharedRef<SWidget>& Widget, const FText& Label, bool bNoIndent = false, bool bSearchable = true);
+
+	bool IsSubMenu() const { return SubMenuData.bIsSubMenu; }
+
 	friend struct FEditorMenuSection;
 	friend class UEditorMenuEntryScript;
 
@@ -113,13 +135,13 @@ public:
 	bool bShouldCloseWindowAfterMenuSelection;
 
 	UPROPERTY(EditAnywhere, BlueprintReadWrite, Category = "Editor UI")
-	FSubMenuEntryData SubMenu;
+	UEditorMenuEntryScript* ScriptObject;
 
-	UPROPERTY(EditAnywhere, BlueprintReadWrite, Category = "Editor UI")
-	FToolBarData ToolBar;
+	FEditorMenuEntrySubMenuData SubMenuData;
 
-	UPROPERTY(EditAnywhere, BlueprintReadWrite, Category = "Editor UI")
-	UEditorMenuEntryScript* ScriptObject;
+	FEditorMenuEntryToolBarData ToolBarData;
+
+	FEditorMenuEntryWidgetData WidgetData;
 
 	/** Optional delegate that returns a widget to use as this menu entry */
 	FNewEditorMenuWidget MakeWidget;
