Updates to Editor MenusBlueprint Library string commands
#rnx

diff --git a/Engine/Source/Editor/EditorMenus/Private/EditorMenuSubsystem.cpp b/Engine/Source/Editor/EditorMenus/Private/EditorMenuSubsystem.cpp
--- a/Engine/Source/Editor/EditorMenus/Private/EditorMenuSubsystem.cpp
+++ b/Engine/Source/Editor/EditorMenus/Private/EditorMenuSubsystem.cpp
@@ -33,8 +33,8 @@ FName FEditorMenuStringCommand::GetTypeName() const
 		return CommandName;
 	case EEditorMenuStringCommandType::Python:
 		return PythonName;
-	case EEditorMenuStringCommandType::Other:
-		return TypeName;
+	case EEditorMenuStringCommandType::Custom:
+		return CustomType;
 	default:
 		break;
 	}
diff --git a/Engine/Source/Editor/EditorMenus/Private/EditorMenusBlueprintLibrary.cpp b/Engine/Source/Editor/EditorMenus/Private/EditorMenusBlueprintLibrary.cpp
--- a/Engine/Source/Editor/EditorMenus/Private/EditorMenusBlueprintLibrary.cpp
+++ b/Engine/Source/Editor/EditorMenus/Private/EditorMenusBlueprintLibrary.cpp
@@ -53,10 +53,11 @@ void UEditorMenuEntryExtensions::SetIcon(UPARAM(ref) FEditorMenuEntry& Target, c
 	Target.Icon = FSlateIcon(StyleSetName, StyleName, SmallStyleName);
 }
 
-void UEditorMenuEntryExtensions::SetStringCommand(UPARAM(ref) FEditorMenuEntry& Target, const FName TypeName, const FString& String)
+void UEditorMenuEntryExtensions::SetStringCommand(UPARAM(ref) FEditorMenuEntry& Target, const EEditorMenuStringCommandType Type, const FString& String, const FName CustomType)
 {
 	Target.ResetActions();
-	Target.StringExecuteAction.TypeName = TypeName;
+	Target.StringExecuteAction.Type = Type;
+	Target.StringExecuteAction.CustomType = CustomType;
 	Target.StringExecuteAction.String = String;
 }
 
diff --git a/Engine/Source/Editor/EditorMenus/Private/EditorMenusBlueprintLibrary.h b/Engine/Source/Editor/EditorMenus/Private/EditorMenusBlueprintLibrary.h
--- a/Engine/Source/Editor/EditorMenus/Private/EditorMenusBlueprintLibrary.h
+++ b/Engine/Source/Editor/EditorMenus/Private/EditorMenusBlueprintLibrary.h
@@ -8,6 +8,7 @@
 #include "EditorMenuEntry.h"
 #include "EditorMenuSection.h"
 #include "Misc/Attribute.h"
+#include "Kismet/BlueprintFunctionLibrary.h"
 
 #include "EditorMenusBlueprintLibrary.generated.h"
 
@@ -51,7 +52,7 @@ public:
 	static void SetIcon(UPARAM(ref) FEditorMenuEntry& Target, const FName StyleSetName, const FName StyleName = NAME_None, const FName SmallStyleName = NAME_None);
 
 	UFUNCTION(BlueprintCallable, Category = "Editor UI", meta = (ScriptMethod))
-	static void SetStringCommand(UPARAM(ref) FEditorMenuEntry& Target, const FName TypeName, const FString& String);
+	static void SetStringCommand(UPARAM(ref) FEditorMenuEntry& Target, const EEditorMenuStringCommandType Type, const FString& String, const FName CustomType = NAME_None);
 
 	UFUNCTION(BlueprintCallable, Category = "Editor UI")
 	static FEditorMenuEntry InitMenuEntry(const FName InOwner, const FName InName, const FText& InLabel, const FText& InToolTip, const FEditorMenuStringCommand& StringCommand);
diff --git a/Engine/Source/Editor/EditorMenus/Public/EditorMenuMisc.h b/Engine/Source/Editor/EditorMenus/Public/EditorMenuMisc.h
--- a/Engine/Source/Editor/EditorMenus/Public/EditorMenuMisc.h
+++ b/Engine/Source/Editor/EditorMenus/Public/EditorMenuMisc.h
@@ -14,7 +14,7 @@ enum class EEditorMenuStringCommandType : uint8
 {
 	Command,
 	Python,
-	Other
+	Custom
 };
 
 USTRUCT(BlueprintType)
@@ -22,12 +22,15 @@ struct EDITORMENUS_API FEditorMenuStringCommand
 {
 	GENERATED_BODY()
 
+	// Which command handler to use
 	UPROPERTY(EditAnywhere, BlueprintReadWrite, Category = "Editor UI")
 	EEditorMenuStringCommandType Type;
 
+	// Which command handler to use when type is custom
 	UPROPERTY(EditAnywhere, BlueprintReadWrite, Category = "Editor UI")
-	FName TypeName;
+	FName CustomType;
 
+	// String to pass to command handler
 	UPROPERTY(EditAnywhere, BlueprintReadWrite, Category = "Editor UI")
 	FString String;
 
@@ -59,9 +62,11 @@ struct FEditorMenuInsert
 	FEditorMenuInsert() : Position(EEditorMenuInsertType::Default) {}
 	FEditorMenuInsert(FName InName, EEditorMenuInsertType InPosition) : Name(InName), Position(InPosition) {}
 
+	// Where to insert
 	UPROPERTY(EditAnywhere, BlueprintReadWrite, Category = "Editor UI")
 	FName Name;
 
+	// How to insert
 	UPROPERTY(EditAnywhere, BlueprintReadWrite, Category = "Editor UI")
 	EEditorMenuInsertType Position;
 
