Fix linux compile warning: conversion from string literal to 'ANSICHAR *'
#rnx

diff --git a/Engine/Source/Editor/EditorMenus/Public/EditorMenuOwner.h b/Engine/Source/Editor/EditorMenus/Public/EditorMenuOwner.h
--- a/Engine/Source/Editor/EditorMenus/Public/EditorMenuOwner.h
+++ b/Engine/Source/Editor/EditorMenus/Public/EditorMenuOwner.h
@@ -30,8 +30,8 @@ public:
 	FORCEINLINE FEditorMenuOwner() : ValueInt64(0), ValueType(EValueType::None) {}
 	FORCEINLINE FEditorMenuOwner(void* InPointer) : ValueInt64(reinterpret_cast<int64>(InPointer)), ValueType(EValueType::Pointer) {}
 
-	FEditorMenuOwner(ANSICHAR* InValue) : FEditorMenuOwner(FName(InValue)) {}
-	FEditorMenuOwner(TCHAR* InValue) : FEditorMenuOwner(FName(InValue)) {}
+	FEditorMenuOwner(const WIDECHAR* InValue) : FEditorMenuOwner(FName(InValue)) {}
+	FEditorMenuOwner(const ANSICHAR* InValue) : FEditorMenuOwner(FName(InValue)) {}
 
 	FEditorMenuOwner(const FName InValue)
 	{
