Fix submenus not going into the correct section
#rnx

diff --git a/Engine/Source/Editor/EditorMenus/Private/EditorMenu.cpp b/Engine/Source/Editor/EditorMenus/Private/EditorMenu.cpp
--- a/Engine/Source/Editor/EditorMenus/Private/EditorMenu.cpp
+++ b/Engine/Source/Editor/EditorMenus/Private/EditorMenu.cpp
@@ -161,7 +161,7 @@ UEditorMenu* UEditorMenu::AddSubMenu(const FEditorMenuOwner InOwner, const FName
 {
 	FEditorMenuEntry Args = FEditorMenuEntry::InitSubMenu(MenuName, InName, InLabel, InToolTip, FNewEditorMenuChoice());
 	Args.Owner = InOwner;
-	FindOrAddSection(NAME_None).AddEntry(Args);
+	FindOrAddSection(SectionName).AddEntry(Args);
 	return UEditorMenuSubsystem::Get()->ExtendMenu(*(MenuName.ToString() + TEXT(".") + InName.ToString()));
 }
 
