Header file
#rnx

diff --git a/Engine/Source/Editor/EditorMenus/Private/EditorMenuSubsystem.cpp b/Engine/Source/Editor/EditorMenus/Private/EditorMenuSubsystem.cpp
--- a/Engine/Source/Editor/EditorMenus/Private/EditorMenuSubsystem.cpp
+++ b/Engine/Source/Editor/EditorMenus/Private/EditorMenuSubsystem.cpp
@@ -7,6 +7,7 @@
 #include "Framework/Commands/UIAction.h"
 #include "Framework/MultiBox/MultiBox.h"
 #include "Framework/MultiBox/MultiBoxBuilder.h"
+#include "Framework/Application/SlateApplication.h"
 #include "Internationalization/Internationalization.h"
 #include "Editor/EditorPerProjectUserSettings.h"
 #include "HAL/PlatformApplicationMisc.h"
