Editor menu api updates for python
#rnx

diff --git a/Engine/Source/Editor/EditorMenus/Private/EditorMenu.cpp b/Engine/Source/Editor/EditorMenus/Private/EditorMenu.cpp
--- a/Engine/Source/Editor/EditorMenus/Private/EditorMenu.cpp
+++ b/Engine/Source/Editor/EditorMenus/Private/EditorMenu.cpp
@@ -157,6 +157,11 @@ void UEditorMenu::AddMenuEntryObject(UEditorMenuEntryScript* InObject)
 	FindOrAddSection(InObject->Data.Section).AddEntryObject(InObject);
 }
 
+UEditorMenu* UEditorMenu::AddSubMenuScript(const FName InOwner, const FName SectionName, const FName InName, const FText& InLabel, const FText& InToolTip)
+{
+	return AddSubMenu(InOwner, SectionName, InName, InLabel, InToolTip);
+}
+
 UEditorMenu* UEditorMenu::AddSubMenu(const FEditorMenuOwner InOwner, const FName SectionName, const FName InName, const FText& InLabel, const FText& InToolTip)
 {
 	FEditorMenuEntry Args = FEditorMenuEntry::InitSubMenu(MenuName, InName, InLabel, InToolTip, FNewEditorMenuChoice());
diff --git a/Engine/Source/Editor/EditorMenus/Private/EditorMenuSubsystem.cpp b/Engine/Source/Editor/EditorMenus/Private/EditorMenuSubsystem.cpp
--- a/Engine/Source/Editor/EditorMenus/Private/EditorMenuSubsystem.cpp
+++ b/Engine/Source/Editor/EditorMenus/Private/EditorMenuSubsystem.cpp
@@ -1236,7 +1236,7 @@ UEditorMenu* UEditorMenuSubsystem::GenerateMenuAsBuilder(const UEditorMenu* InMe
 	return GenerateMenu(Hierarchy, InMenuContext);
 }
 
-UEditorMenu* UEditorMenuSubsystem::RegisterMenu(const FName InName, const FName InParent, EMultiBoxType InType)
+UEditorMenu* UEditorMenuSubsystem::RegisterMenu(const FName InName, const FName InParent, EMultiBoxType InType, bool bWarnIfAlreadyRegistered)
 {
 	if (UEditorMenu* Found = FindMenu(InName))
 	{
@@ -1247,7 +1247,7 @@ UEditorMenu* UEditorMenuSubsystem::RegisterMenu(const FName InName, const FName
 			Found->MenuOwner = CurrentOwner();
 			Found->bRegistered = true;
 		}
-		else
+		else if (bWarnIfAlreadyRegistered)
 		{
 			UE_LOG(LogEditorMenus, Warning, TEXT("Menu already registered : %s"), *InName.ToString());
 		}
diff --git a/Engine/Source/Editor/EditorMenus/Private/EditorMenusBlueprintLibrary.cpp b/Engine/Source/Editor/EditorMenus/Private/EditorMenusBlueprintLibrary.cpp
--- a/Engine/Source/Editor/EditorMenus/Private/EditorMenusBlueprintLibrary.cpp
+++ b/Engine/Source/Editor/EditorMenus/Private/EditorMenusBlueprintLibrary.cpp
@@ -3,6 +3,18 @@
 
 #include "EditorMenusBlueprintLibrary.h"
 
+FEditorMenuStringCommand UEditorMenuEntryExtensions::MakeStringCommand(EEditorMenuStringCommandType Type, FName CustomType, const FString& String)
+{
+	return FEditorMenuStringCommand(Type, CustomType, String);
+}
+
+void UEditorMenuEntryExtensions::BreakStringCommand(const FEditorMenuStringCommand& InValue, EEditorMenuStringCommandType& Type, FName& CustomType, FString& String)
+{
+	Type = InValue.Type;
+	CustomType = InValue.CustomType;
+	String = InValue.String;
+}
+
 FEditorMenuOwner UEditorMenuEntryExtensions::MakeEditorMenuOwner(FName Name)
 {
 	return FEditorMenuOwner(Name);
@@ -53,7 +65,7 @@ void UEditorMenuEntryExtensions::SetIcon(UPARAM(ref) FEditorMenuEntry& Target, c
 	Target.Icon = FSlateIcon(StyleSetName, StyleName, SmallStyleName);
 }
 
-void UEditorMenuEntryExtensions::SetStringCommand(UPARAM(ref) FEditorMenuEntry& Target, const EEditorMenuStringCommandType Type, const FString& String, const FName CustomType)
+void UEditorMenuEntryExtensions::SetStringCommand(UPARAM(ref) FEditorMenuEntry& Target, const EEditorMenuStringCommandType Type, const FName CustomType, const FString& String)
 {
 	Target.ResetActions();
 	Target.StringExecuteAction.Type = Type;
@@ -61,12 +73,12 @@ void UEditorMenuEntryExtensions::SetStringCommand(UPARAM(ref) FEditorMenuEntry&
 	Target.StringExecuteAction.String = String;
 }
 
-FEditorMenuEntry UEditorMenuEntryExtensions::InitMenuEntry(const FName InOwner, const FName InName, const FText& InLabel, const FText& InToolTip, const FEditorMenuStringCommand& StringCommand)
+FEditorMenuEntry UEditorMenuEntryExtensions::InitMenuEntry(const FName InOwner, const FName InName, const FText& InLabel, const FText& InToolTip, const EEditorMenuStringCommandType CommandType, const FName CustomCommandType, const FString& CommandString)
 {
 	FEditorMenuEntry Entry(InOwner, InName, EMultiBlockType::MenuEntry);
 	Entry.Label = InLabel;
 	Entry.ToolTip = InToolTip;
-	Entry.StringExecuteAction = StringCommand;
+	Entry.StringExecuteAction = FEditorMenuStringCommand(CommandType, CustomCommandType, CommandString);
 	return Entry;
 }
 
diff --git a/Engine/Source/Editor/EditorMenus/Private/EditorMenusBlueprintLibrary.h b/Engine/Source/Editor/EditorMenus/Private/EditorMenusBlueprintLibrary.h
--- a/Engine/Source/Editor/EditorMenus/Private/EditorMenusBlueprintLibrary.h
+++ b/Engine/Source/Editor/EditorMenus/Private/EditorMenusBlueprintLibrary.h
@@ -7,6 +7,7 @@
 #include "EditorMenuDelegates.h"
 #include "EditorMenuEntry.h"
 #include "EditorMenuSection.h"
+#include "EditorMenuMisc.h"
 #include "Misc/Attribute.h"
 #include "Kismet/BlueprintFunctionLibrary.h"
 
@@ -30,6 +31,12 @@ class EDITORMENUS_API UEditorMenuEntryExtensions : public UBlueprintFunctionLibr
 
 public:
 
+	UFUNCTION(BlueprintPure, Category = "Editor UI", meta = (Keywords = "construct build", NativeMakeFunc))
+	static FEditorMenuStringCommand MakeStringCommand(EEditorMenuStringCommandType Type, FName CustomType, const FString& String);
+
+	UFUNCTION(BlueprintPure, Category = "Editor UI", meta = (NativeBreakFunc))
+	static void BreakStringCommand(const FEditorMenuStringCommand& InValue, EEditorMenuStringCommandType& Type, FName& CustomType, FString& String);
+
 	UFUNCTION(BlueprintPure, Category = "Editor UI", meta = (Keywords = "construct build", NativeMakeFunc))
 	static FEditorMenuOwner MakeEditorMenuOwner(FName Name);
 
@@ -52,10 +59,10 @@ public:
 	static void SetIcon(UPARAM(ref) FEditorMenuEntry& Target, const FName StyleSetName, const FName StyleName = NAME_None, const FName SmallStyleName = NAME_None);
 
 	UFUNCTION(BlueprintCallable, Category = "Editor UI", meta = (ScriptMethod))
-	static void SetStringCommand(UPARAM(ref) FEditorMenuEntry& Target, const EEditorMenuStringCommandType Type, const FString& String, const FName CustomType = NAME_None);
+	static void SetStringCommand(UPARAM(ref) FEditorMenuEntry& Target, const EEditorMenuStringCommandType Type, const FName CustomType, const FString& String);
 
 	UFUNCTION(BlueprintCallable, Category = "Editor UI")
-	static FEditorMenuEntry InitMenuEntry(const FName InOwner, const FName InName, const FText& InLabel, const FText& InToolTip, const FEditorMenuStringCommand& StringCommand);
+	static FEditorMenuEntry InitMenuEntry(const FName InOwner, const FName InName, const FText& InLabel, const FText& InToolTip, const EEditorMenuStringCommandType CommandType, const FName CustomCommandType, const FString& CommandString);
 };
 
 UCLASS()
diff --git a/Engine/Source/Editor/EditorMenus/Public/EditorMenu.h b/Engine/Source/Editor/EditorMenus/Public/EditorMenu.h
--- a/Engine/Source/Editor/EditorMenus/Public/EditorMenu.h
+++ b/Engine/Source/Editor/EditorMenus/Public/EditorMenu.h
@@ -37,7 +37,9 @@ public:
 	UFUNCTION(BlueprintCallable, Category = "Editor UI")
 	void AddMenuEntryObject(UEditorMenuEntryScript* InObject);
 
-	UFUNCTION(BlueprintCallable, Category = "Editor UI", meta = ( AutoCreateRefTerm = "Label,ToolTip" ))
+	UFUNCTION(BlueprintCallable, Category = "Editor UI", meta = ( ScriptName = "AddSubMenu", AutoCreateRefTerm = "Label,ToolTip" ))
+	UEditorMenu* AddSubMenuScript(const FName Owner, const FName SectionName, const FName Name, const FText& Label, const FText& ToolTip = FText());
+
 	UEditorMenu* AddSubMenu(const FEditorMenuOwner Owner, const FName SectionName, const FName Name, const FText& Label, const FText& ToolTip = FText());
 
 	void RemoveSection(const FName SectionName);
diff --git a/Engine/Source/Editor/EditorMenus/Public/EditorMenuMisc.h b/Engine/Source/Editor/EditorMenus/Public/EditorMenuMisc.h
--- a/Engine/Source/Editor/EditorMenus/Public/EditorMenuMisc.h
+++ b/Engine/Source/Editor/EditorMenus/Public/EditorMenuMisc.h
@@ -17,11 +17,20 @@ enum class EEditorMenuStringCommandType : uint8
 	Custom
 };
 
-USTRUCT(BlueprintType)
+USTRUCT(BlueprintType, meta=(HasNativeBreak="EditorMenus.EditorMenuEntryExtensions.BreakStringCommand", HasNativeMake="EditorMenus.EditorMenuEntryExtensions.MakeStringCommand"))
 struct EDITORMENUS_API FEditorMenuStringCommand
 {
 	GENERATED_BODY()
 
+	FEditorMenuStringCommand() : Type(EEditorMenuStringCommandType::Command) {}
+
+	FEditorMenuStringCommand(EEditorMenuStringCommandType InType, FName InCustomType, const FString& InString) :
+		Type(InType),
+		CustomType(InCustomType),
+		String(InString)
+	{
+	}
+
 	// Which command handler to use
 	UPROPERTY(EditAnywhere, BlueprintReadWrite, Category = "Editor UI")
 	EEditorMenuStringCommandType Type;
diff --git a/Engine/Source/Editor/EditorMenus/Public/EditorMenuSubsystem.h b/Engine/Source/Editor/EditorMenus/Public/EditorMenuSubsystem.h
--- a/Engine/Source/Editor/EditorMenus/Public/EditorMenuSubsystem.h
+++ b/Engine/Source/Editor/EditorMenus/Public/EditorMenuSubsystem.h
@@ -134,10 +134,11 @@ public:
 	 * Registers a menu by name
 	 * @param	Parent	Optional name of a menu to layer on top of.
 	 * @param	Type	Type of menu that will be generated such as: ToolBar, VerticalToolBar, etc..
+	 * @param	bWarnIfAlreadyRegistered	Display warning if already registered
 	 * @return	EditorMenu	Menu object
 	 */
 	UFUNCTION(BlueprintCallable, Category = "Editor UI")
-	UEditorMenu* RegisterMenu(FName Name, const FName Parent = NAME_None, EMultiBoxType Type = EMultiBoxType::Menu);
+	UEditorMenu* RegisterMenu(FName Name, const FName Parent = NAME_None, EMultiBoxType Type = EMultiBoxType::Menu, bool bWarnIfAlreadyRegistered = true);
 
 	/**
 	 * Extends a menu without registering the menu or claiming ownership of it. Ok to call even if menu does not exist yet.
