Python based menu tests
#rnx

diff --git a/Engine/Source/Editor/EditorMenus/Private/EditorMenuEntryScript.cpp b/Engine/Source/Editor/EditorMenus/Private/EditorMenuEntryScript.cpp
--- a/Engine/Source/Editor/EditorMenus/Private/EditorMenuEntryScript.cpp
+++ b/Engine/Source/Editor/EditorMenus/Private/EditorMenuEntryScript.cpp
@@ -69,6 +69,36 @@ TAttribute<FText> UEditorMenuEntryScript::CreateToolTipAttribute(FEditorMenuCont
 	return Data.ToolTip;
 }
 
+TAttribute<FSlateIcon> UEditorMenuEntryScript::CreateIconAttribute(FEditorMenuContext& Context)
+{
+	static const FName FunctionName = GET_FUNCTION_NAME_CHECKED(UEditorMenuEntryScript, GetIcon);
+	if (GetClass()->IsFunctionImplementedInScript(FunctionName))
+	{
+		TWeakObjectPtr<UEditorMenuEntryScript> WeakThis(this);
+		TAttribute<FSlateIcon>::FGetter Getter;
+		Getter.BindLambda([=]()
+		{
+			if (UEditorMenuEntryScript* Object = WeakThis.Get())
+			{
+				return Object->GetIcon(Context).GetSlateIcon();
+			}
+			else
+			{
+				return FSlateIcon();
+			}
+		});
+
+		return TAttribute<FSlateIcon>::Create(Getter);
+	}
+
+	return Data.Icon.GetSlateIcon();
+}
+
+FSlateIcon UEditorMenuEntryScript::GetSlateIcon(const FEditorMenuContext& Context) const
+{
+	return GetIcon(Context).GetSlateIcon();
+}
+
 void UEditorMenuEntryScript::RegisterMenuEntry()
 {
 	UEditorMenuSubsystem::AddMenuEntryObject(this);
@@ -115,6 +145,8 @@ void UEditorMenuEntryScript::ToMenuEntry(FEditorMenuEntry& Output)
 		else
 		{
 			Output = FEditorMenuEntry::InitMenuEntry(Data.Name, Data.Label, Data.ToolTip, Data.Icon, FUIAction());
+			Output.UserInterfaceActionType = Data.Advanced.UserInterfaceActionType;
+			Output.TutorialHighlightName = Data.Advanced.TutorialHighlight;
 		}
 	}
 
@@ -122,7 +154,7 @@ void UEditorMenuEntryScript::ToMenuEntry(FEditorMenuEntry& Output)
 	{
 		Output.InsertPosition = Data.InsertPosition;
 	}
-	
+
 	Output.ScriptObject = this;
 
 	Output.Owner = Data.OwnerName;
diff --git a/Engine/Source/Editor/EditorMenus/Private/EditorMenuSubsystem.cpp b/Engine/Source/Editor/EditorMenus/Private/EditorMenuSubsystem.cpp
--- a/Engine/Source/Editor/EditorMenus/Private/EditorMenuSubsystem.cpp
+++ b/Engine/Source/Editor/EditorMenus/Private/EditorMenuSubsystem.cpp
@@ -181,7 +181,7 @@ void UEditorMenuSubsystem::AssembleMenuSection(UEditorMenu* GeneratedMenu, const
 
 			if (Block.IsScriptObjectDynamicConstruct())
 			{
-				Block.ScriptObject->ConstructMenuEntry(*DestSection, DestSection->Context);
+				Block.ScriptObject->ConstructMenuEntry(Constructed, DestSection->Name, DestSection->Context);
 			}
 			else
 			{
@@ -594,7 +594,8 @@ void UEditorMenuSubsystem::PopulateMenuBuilder(FMenuBuilder& MenuBuilder, UEdito
 					else if (Block.ScriptObject)
 					{
 						UEditorMenuEntryScript* ScriptObject = Block.ScriptObject;
-						MenuBuilder.AddMenuEntry(ScriptObject->CreateLabelAttribute(MenuData->Context), ScriptObject->CreateToolTipAttribute(MenuData->Context), Block.Icon.Get(), UIAction, ScriptObject->Data.Name, Block.UserInterfaceActionType, Block.TutorialHighlightName);
+						const FSlateIcon Icon = ScriptObject->CreateIconAttribute(MenuData->Context).Get();
+						MenuBuilder.AddMenuEntry(ScriptObject->CreateLabelAttribute(MenuData->Context), ScriptObject->CreateToolTipAttribute(MenuData->Context), Icon, UIAction, ScriptObject->Data.Name, Block.UserInterfaceActionType, Block.TutorialHighlightName);
 					}
 					else
 					{
@@ -666,7 +667,8 @@ void UEditorMenuSubsystem::PopulateToolBarBuilder(FToolBarBuilder& ToolBarBuilde
 				else if (Block.ScriptObject)
 				{
 					UEditorMenuEntryScript* ScriptObject = Block.ScriptObject;
-					ToolBarBuilder.AddToolBarButton(UIAction, ScriptObject->Data.Name, ScriptObject->CreateLabelAttribute(MenuData->Context), ScriptObject->CreateToolTipAttribute(MenuData->Context), Block.Icon, Block.UserInterfaceActionType, Block.TutorialHighlightName);
+					TAttribute<FSlateIcon> Icon = ScriptObject->CreateIconAttribute(MenuData->Context);
+					ToolBarBuilder.AddToolBarButton(UIAction, ScriptObject->Data.Name, ScriptObject->CreateLabelAttribute(MenuData->Context), ScriptObject->CreateToolTipAttribute(MenuData->Context), Icon, Block.UserInterfaceActionType, Block.TutorialHighlightName);
 				}
 				else
 				{
diff --git a/Engine/Source/Editor/EditorMenus/Private/EditorMenusBlueprintLibrary.cpp b/Engine/Source/Editor/EditorMenus/Private/EditorMenusBlueprintLibrary.cpp
--- a/Engine/Source/Editor/EditorMenus/Private/EditorMenusBlueprintLibrary.cpp
+++ b/Engine/Source/Editor/EditorMenus/Private/EditorMenusBlueprintLibrary.cpp
@@ -3,6 +3,25 @@
 
 #include "EditorMenusBlueprintLibrary.h"
 
+FScriptSlateIcon UEditorMenuEntryExtensions::MakeScriptSlateIcon(const FName StyleSetName, const FName StyleName, const FName SmallStyleName)
+{
+	if (SmallStyleName == NAME_None)
+	{
+		return FScriptSlateIcon(StyleSetName, StyleName);
+	}
+	else
+	{
+		return FScriptSlateIcon(StyleSetName, StyleName, SmallStyleName);
+	}
+}
+
+void UEditorMenuEntryExtensions::BreakScriptSlateIcon(const FScriptSlateIcon& InValue, FName& StyleSetName, FName& StyleName, FName& SmallStyleName)
+{
+	StyleSetName = InValue.StyleSetName;
+	StyleName = InValue.StyleName;
+	SmallStyleName = InValue.SmallStyleName;
+}
+
 FEditorMenuStringCommand UEditorMenuEntryExtensions::MakeStringCommand(EEditorMenuStringCommandType Type, FName CustomType, const FString& String)
 {
 	return FEditorMenuStringCommand(Type, CustomType, String);
diff --git a/Engine/Source/Editor/EditorMenus/Private/EditorMenusBlueprintLibrary.h b/Engine/Source/Editor/EditorMenus/Private/EditorMenusBlueprintLibrary.h
--- a/Engine/Source/Editor/EditorMenus/Private/EditorMenusBlueprintLibrary.h
+++ b/Engine/Source/Editor/EditorMenus/Private/EditorMenusBlueprintLibrary.h
@@ -8,6 +8,7 @@
 #include "EditorMenuEntry.h"
 #include "EditorMenuSection.h"
 #include "EditorMenuMisc.h"
+#include "EditorMenuEntryScript.h"
 #include "Misc/Attribute.h"
 #include "Kismet/BlueprintFunctionLibrary.h"
 
@@ -31,6 +32,12 @@ class EDITORMENUS_API UEditorMenuEntryExtensions : public UBlueprintFunctionLibr
 
 public:
 
+	UFUNCTION(BlueprintPure, Category = "Editor UI", meta = (Keywords = "construct build", NativeMakeFunc))
+	static FScriptSlateIcon MakeScriptSlateIcon(const FName StyleSetName, const FName StyleName, const FName SmallStyleName = NAME_None);
+
+	UFUNCTION(BlueprintPure, Category = "Editor UI", meta = (NativeBreakFunc))
+	static void BreakScriptSlateIcon(const FScriptSlateIcon& InValue, FName& StyleSetName, FName& StyleName, FName& SmallStyleName);
+
 	UFUNCTION(BlueprintPure, Category = "Editor UI", meta = (Keywords = "construct build", NativeMakeFunc))
 	static FEditorMenuStringCommand MakeStringCommand(EEditorMenuStringCommandType Type, FName CustomType, const FString& String);
 
diff --git a/Engine/Source/Editor/EditorMenus/Public/EditorMenuEntryScript.h b/Engine/Source/Editor/EditorMenus/Public/EditorMenuEntryScript.h
--- a/Engine/Source/Editor/EditorMenus/Public/EditorMenuEntryScript.h
+++ b/Engine/Source/Editor/EditorMenus/Public/EditorMenuEntryScript.h
@@ -20,7 +20,7 @@
 struct FEditorMenuEntry;
 struct FEditorMenuSection;
 
-USTRUCT(BlueprintType)
+USTRUCT(BlueprintType, meta=(HasNativeBreak="EditorMenus.EditorMenuEntryExtensions.BreakScriptSlateIcon", HasNativeMake="EditorMenus.EditorMenuEntryExtensions.MakeScriptSlateIcon"))
 struct EDITORMENUS_API FScriptSlateIcon
 {
 	GENERATED_BODY()
@@ -146,7 +146,7 @@ public:
 	virtual FScriptSlateIcon GetIcon_Implementation(const FEditorMenuContext& Context) const { return Data.Icon; }
 
 	UFUNCTION(BlueprintImplementableEvent, Category = "Advanced")
-	void ConstructMenuEntry(FEditorMenuSection& Section, const FEditorMenuContext& Context);
+	void ConstructMenuEntry(UEditorMenu* Menu, const FName SectionName, const FEditorMenuContext& Context);
 
 	UFUNCTION(BlueprintCallable, Category = "Advanced")
 	void RegisterMenuEntry();
@@ -163,10 +163,14 @@ private:
 
 	TAttribute<FText> CreateToolTipAttribute(FEditorMenuContext& Context);
 
+	TAttribute<FSlateIcon> CreateIconAttribute(FEditorMenuContext& Context);
+
 	void ToMenuEntry(FEditorMenuEntry& Output);
 
 	bool IsDynamicConstruct() const;
 
+	FSlateIcon GetSlateIcon(const FEditorMenuContext& Context) const;
+
 public:
 
 	UPROPERTY(EditAnywhere, BlueprintReadWrite, Category = "Data")
