Fix menu context for dynamic entries
#rnx

diff --git a/Engine/Source/Editor/EditorMenus/Private/EditorMenuSubsystem.cpp b/Engine/Source/Editor/EditorMenus/Private/EditorMenuSubsystem.cpp
--- a/Engine/Source/Editor/EditorMenus/Private/EditorMenuSubsystem.cpp
+++ b/Engine/Source/Editor/EditorMenus/Private/EditorMenuSubsystem.cpp
@@ -186,6 +186,7 @@ void UEditorMenuSubsystem::AssembleMenuSection(UEditorMenu* GeneratedMenu, const
 			else
 			{
 				FEditorMenuSection& ConstructedSection = Constructed->AddSection(DestSection->Name);
+				ConstructedSection.Context = Constructed->Context;
 				Block.Construct.Execute(ConstructedSection);
 			}
 
