Each individual menu entry can now have a command list.
#rnx

diff --git a/Engine/Source/Editor/EditorMenus/Private/EditorMenuContext.cpp b/Engine/Source/Editor/EditorMenus/Private/EditorMenuContext.cpp
--- a/Engine/Source/Editor/EditorMenus/Private/EditorMenuContext.cpp
+++ b/Engine/Source/Editor/EditorMenus/Private/EditorMenuContext.cpp
@@ -83,6 +83,23 @@ void FEditorMenuContext::AppendCommandList(const TSharedPtr<FUICommandList>& InC
 	}
 }
 
+const FUIAction* FEditorMenuContext::GetActionForCommand(TSharedPtr<const FUICommandInfo> Command, TSharedPtr<const FUICommandList>& OutCommandList) const
+{
+	for (const TSharedPtr<FUICommandList>& CommandListIter : CommandLists)
+	{
+		if (CommandListIter.IsValid())
+		{
+			if (const FUIAction* Result = CommandListIter->GetActionForCommand(Command))
+			{
+				OutCommandList = CommandListIter;
+				return Result;
+			}
+		}
+	}
+
+	return nullptr;
+}
+
 void FEditorMenuContext::AddExtender(const TSharedPtr<FExtender>& InExtender)
 {
 	if (!ExtensibilityManager.IsValid())
diff --git a/Engine/Source/Editor/EditorMenus/Private/EditorMenuEntry.cpp b/Engine/Source/Editor/EditorMenus/Private/EditorMenuEntry.cpp
--- a/Engine/Source/Editor/EditorMenus/Private/EditorMenuEntry.cpp
+++ b/Engine/Source/Editor/EditorMenus/Private/EditorMenuEntry.cpp
@@ -31,7 +31,34 @@ FEditorMenuEntry::FEditorMenuEntry(const FEditorMenuOwner InOwner, const FName I
 {
 }
 
-void FEditorMenuEntry::SetCommand(const TSharedPtr< const FUICommandInfo >& InCommand, FName InName, const TAttribute<FText>& InLabel, const TAttribute<FText>& InToolTip, const TAttribute<FSlateIcon>& InIcon)
+const FUIAction* FEditorMenuEntry::GetActionForCommand(const FEditorMenuContext& InContext, TSharedPtr<const FUICommandList>& OutCommandList) const
+{
+	if (Command.IsValid())
+	{
+		if (CommandList.IsValid())
+		{
+			const FUIAction* Result = CommandList->GetActionForCommand(Command);
+			if (Result)
+			{
+				OutCommandList = CommandList;
+				return Result;
+			}
+		}
+		else
+		{
+			return InContext.GetActionForCommand(Command, OutCommandList);
+		}
+	}
+
+	return nullptr;
+}
+
+void FEditorMenuEntry::SetCommandList(const TSharedPtr<const FUICommandList>& InCommandList)
+{
+	CommandList = InCommandList;
+}
+
+void FEditorMenuEntry::SetCommand(const TSharedPtr<const FUICommandInfo>& InCommand, FName InName, const TAttribute<FText>& InLabel, const TAttribute<FText>& InToolTip, const TAttribute<FSlateIcon>& InIcon)
 {
 	Command = InCommand;
 	Name = InName != NAME_None ? InName : InCommand->GetCommandName();
@@ -57,6 +84,16 @@ FEditorMenuEntry FEditorMenuEntry::InitMenuEntry(const TSharedPtr< const FUIComm
 	FEditorMenuEntry Entry(UEditorMenuSubsystem::Get()->CurrentOwner(), InName, EMultiBlockType::MenuEntry);
 	Entry.TutorialHighlightName = InTutorialHighlightName;
 	Entry.SetCommand(InCommand, InName, InLabel, InToolTip, InIcon);
+	Entry.CommandList.Reset();
+	return Entry;
+}
+
+FEditorMenuEntry FEditorMenuEntry::InitMenuEntryWithCommandList(const TSharedPtr< const FUICommandInfo >& InCommand, const TSharedPtr< const FUICommandList >& InCommandList, const TAttribute<FText>& InLabel, const TAttribute<FText>& InToolTip, const TAttribute<FSlateIcon>& InIcon, const FName InTutorialHighlightName, const FName InName)
+{
+	FEditorMenuEntry Entry(UEditorMenuSubsystem::Get()->CurrentOwner(), InName, EMultiBlockType::MenuEntry);
+	Entry.TutorialHighlightName = InTutorialHighlightName;
+	Entry.SetCommand(InCommand, InName, InLabel, InToolTip, InIcon);
+	Entry.CommandList = InCommandList;
 	return Entry;
 }
 
diff --git a/Engine/Source/Editor/EditorMenus/Private/EditorMenuSection.cpp b/Engine/Source/Editor/EditorMenus/Private/EditorMenuSection.cpp
--- a/Engine/Source/Editor/EditorMenus/Private/EditorMenuSection.cpp
+++ b/Engine/Source/Editor/EditorMenus/Private/EditorMenuSection.cpp
@@ -81,6 +81,11 @@ FEditorMenuEntry& FEditorMenuSection::AddMenuEntry(const TSharedPtr< const FUICo
 	return AddEntry(FEditorMenuEntry::InitMenuEntry(InCommand, InLabelOverride, InToolTipOverride, InIconOverride, InTutorialHighlightName, InNameOverride));
 }
 
+FEditorMenuEntry& FEditorMenuSection::AddMenuEntryWithCommandList(const TSharedPtr< const FUICommandInfo >& InCommand, const TSharedPtr< const FUICommandList >& InCommandList, const TAttribute<FText>& InLabelOverride, const TAttribute<FText>& InToolTipOverride, const TAttribute<FSlateIcon>& InIconOverride, const FName InTutorialHighlightName, const FName InNameOverride)
+{
+	return AddEntry(FEditorMenuEntry::InitMenuEntryWithCommandList(InCommand, InCommandList, InLabelOverride, InToolTipOverride, InIconOverride, InTutorialHighlightName, InNameOverride));
+}
+
 FEditorMenuEntry& FEditorMenuSection::AddDynamicEntry(const FName InName, const FNewEditorMenuSectionDelegate& InConstruct)
 {
 	FEditorMenuEntry& Entry = AddEntry(FEditorMenuEntry(UEditorMenuSubsystem::Get()->CurrentOwner(), InName, EMultiBlockType::MenuEntry));
diff --git a/Engine/Source/Editor/EditorMenus/Private/EditorMenuSubsystem.cpp b/Engine/Source/Editor/EditorMenus/Private/EditorMenuSubsystem.cpp
--- a/Engine/Source/Editor/EditorMenus/Private/EditorMenuSubsystem.cpp
+++ b/Engine/Source/Editor/EditorMenus/Private/EditorMenuSubsystem.cpp
@@ -512,12 +512,13 @@ void UEditorMenuSubsystem::AssembleMenuHierarchy(UEditorMenu* GeneratedMenu, con
 	ApplyCustomization(GeneratedMenu);
 }
 
-void UEditorMenuSubsystem::FillMenuDynamic(FMenuBuilder& MenuBuilder, FNewEditorMenuDelegate InConstructMenu)
+void UEditorMenuSubsystem::FillMenuDynamic(FMenuBuilder& MenuBuilder, FNewEditorMenuDelegate InConstructMenu, const FEditorMenuContext Context)
 {
 	if (InConstructMenu.IsBound())
 	{
 		// Create final menu
 		UEditorMenu* MenuData = NewObject<UEditorMenu>();
+		MenuData->Context = Context;
 		InConstructMenu.Execute(MenuData);
 
 		// Populate menu builder with final menu
@@ -618,7 +619,7 @@ void UEditorMenuSubsystem::PopulateMenuBuilder(FMenuBuilder& MenuBuilder, UEdito
 						MenuBuilder.AddSubMenu(
 							Block.Label,
 							Block.ToolTip,
-							FNewMenuDelegate::CreateUObject(this, &UEditorMenuSubsystem::FillMenuDynamic, Block.SubMenuData.ConstructMenu.NewEditorMenuDelegate),
+							FNewMenuDelegate::CreateUObject(this, &UEditorMenuSubsystem::FillMenuDynamic, Block.SubMenuData.ConstructMenu.NewEditorMenuDelegate, MenuData->Context),
 							Block.SubMenuData.bOpenSubMenuOnClick,
 							Block.Icon.Get(),
 							Block.bShouldCloseWindowAfterMenuSelection,
@@ -661,7 +662,24 @@ void UEditorMenuSubsystem::PopulateMenuBuilder(FMenuBuilder& MenuBuilder, UEdito
 				{
 					if (Block.Command.IsValid())
 					{
+						bool bPopCommandList = false;
+						TSharedPtr<const FUICommandList> CommandListForAction;
+						if (Block.GetActionForCommand(MenuData->Context, CommandListForAction) != nullptr && CommandListForAction.IsValid())
+						{
+							MenuBuilder.PushCommandList(CommandListForAction.ToSharedRef());
+							bPopCommandList = true;
+						}
+						else
+						{
+							UE_LOG(LogEditorMenus, Error, TEXT("UI command not found for menu entry: %s, menu: %s"), *Block.Name.ToString(), *MenuData->MenuName.ToString());
+						}
+
 						MenuBuilder.AddMenuEntry(Block.Command, Block.Name, Block.Label, Block.ToolTip, Block.Icon.Get());
+
+						if (bPopCommandList)
+						{
+							MenuBuilder.PopCommandList();
+						}
 					}
 					else if (Block.ScriptObject)
 					{
@@ -734,7 +752,24 @@ void UEditorMenuSubsystem::PopulateToolBarBuilder(FToolBarBuilder& ToolBarBuilde
 			{
 				if (Block.Command.IsValid())
 				{
+					bool bPopCommandList = false;
+					TSharedPtr<const FUICommandList> CommandListForAction;
+					if (Block.GetActionForCommand(MenuData->Context, CommandListForAction) != nullptr && CommandListForAction.IsValid())
+					{
+						ToolBarBuilder.PushCommandList(CommandListForAction.ToSharedRef());
+						bPopCommandList = true;
+					}
+					else
+					{
+						UE_LOG(LogEditorMenus, Error, TEXT("UI command not found for toolbar entry: %s, toolbar: %s"), *Block.Name.ToString(), *MenuData->MenuName.ToString());
+					}
+
 					ToolBarBuilder.AddToolBarButton(Block.Command, Block.Name, Block.Label, Block.ToolTip, Block.Icon, Block.TutorialHighlightName);
+
+					if (bPopCommandList)
+					{
+						ToolBarBuilder.PopCommandList();
+					}
 				}
 				else if (Block.ScriptObject)
 				{
diff --git a/Engine/Source/Editor/EditorMenus/Public/EditorMenuContext.h b/Engine/Source/Editor/EditorMenus/Public/EditorMenuContext.h
--- a/Engine/Source/Editor/EditorMenus/Public/EditorMenuContext.h
+++ b/Engine/Source/Editor/EditorMenus/Public/EditorMenuContext.h
@@ -54,6 +54,8 @@ public:
 
 	void AppendCommandList(const TSharedRef<FUICommandList>& InCommandList);
 	void AppendCommandList(const TSharedPtr<FUICommandList>& InCommandList);
+	const FUIAction* GetActionForCommand(TSharedPtr<const FUICommandInfo> Command, TSharedPtr<const FUICommandList>& OutCommandList) const;
+
 	void AddExtender(const TSharedPtr<FExtender>& InExtender);
 	TSharedPtr<FExtender> GetAllExtenders();
 	void ReplaceExtenders(const TSharedPtr<FExtender>& InExtender);
@@ -63,7 +65,7 @@ public:
 	void AddObject(UObject* InObject);
 
 	friend class UEditorMenuSubsystem;
-	
+	friend struct FEditorMenuEntry;
 
 private:
 	
diff --git a/Engine/Source/Editor/EditorMenus/Public/EditorMenuEntry.h b/Engine/Source/Editor/EditorMenus/Public/EditorMenuEntry.h
--- a/Engine/Source/Editor/EditorMenus/Public/EditorMenuEntry.h
+++ b/Engine/Source/Editor/EditorMenus/Public/EditorMenuEntry.h
@@ -86,6 +86,7 @@ struct EDITORMENUS_API FEditorMenuEntry
 
 	static FEditorMenuEntry InitMenuEntry(const FName InName, const TAttribute<FText>& InLabel, const TAttribute<FText>& InToolTip, const TAttribute<FSlateIcon>& InIcon, const FEditorUIActionChoice& InAction, const EUserInterfaceActionType UserInterfaceActionType = EUserInterfaceActionType::Button, const FName InTutorialHighlightName = NAME_None);
 	static FEditorMenuEntry InitMenuEntry(const TSharedPtr< const FUICommandInfo >& InCommand, const TAttribute<FText>& InLabelOverride = TAttribute<FText>(), const TAttribute<FText>& InToolTipOverride = TAttribute<FText>(), const TAttribute<FSlateIcon>& InIconOverride = TAttribute<FSlateIcon>(), const FName InTutorialHighlightName = NAME_None, const FName InNameOverride = NAME_None);
+	static FEditorMenuEntry InitMenuEntryWithCommandList(const TSharedPtr<const FUICommandInfo>& InCommand, const TSharedPtr<const FUICommandList>& InCommandList, const TAttribute<FText>& InLabelOverride = TAttribute<FText>(), const TAttribute<FText>& InToolTipOverride = TAttribute<FText>(), const TAttribute<FSlateIcon>& InIconOverride = TAttribute<FSlateIcon>(), const FName InTutorialHighlightName = NAME_None, const FName InNameOverride = NAME_None);
 	static FEditorMenuEntry InitMenuEntry(const FName InName, const FEditorUIActionChoice& InAction, const TSharedRef<SWidget>& Widget);
 
 	static FEditorMenuEntry InitSubMenu(const FName InParentMenu, const FName InName, const TAttribute<FText>& InLabel, const TAttribute<FText>& InToolTip, const FNewEditorMenuChoice& InMakeMenu, bool bInOpenSubMenuOnClick = false, const TAttribute<FSlateIcon>& InIcon = TAttribute<FSlateIcon>(), const bool bShouldCloseWindowAfterMenuSelection = true);
@@ -102,6 +103,10 @@ struct EDITORMENUS_API FEditorMenuEntry
 
 	bool IsSubMenu() const { return SubMenuData.bIsSubMenu; }
 
+	const FUIAction* GetActionForCommand(const FEditorMenuContext& InContext, TSharedPtr<const FUICommandList>& OutCommandList) const;
+
+	void SetCommandList(const TSharedPtr<const FUICommandList>& InCommandList);
+
 	friend struct FEditorMenuSection;
 	friend class UEditorMenuEntryScript;
 
diff --git a/Engine/Source/Editor/EditorMenus/Public/EditorMenuSection.h b/Engine/Source/Editor/EditorMenus/Public/EditorMenuSection.h
--- a/Engine/Source/Editor/EditorMenus/Public/EditorMenuSection.h
+++ b/Engine/Source/Editor/EditorMenus/Public/EditorMenuSection.h
@@ -36,6 +36,7 @@ public:
 	FEditorMenuEntry& AddEntry(const FEditorMenuEntry& Args);
 	FEditorMenuEntry& AddMenuEntry(const FName InName, const TAttribute<FText>& InLabel, const TAttribute<FText>& InToolTip, const TAttribute<FSlateIcon>& InIcon, const FEditorUIActionChoice& InAction, const EUserInterfaceActionType UserInterfaceActionType = EUserInterfaceActionType::Button, const FName InTutorialHighlightName = NAME_None);
 	FEditorMenuEntry& AddMenuEntry(const TSharedPtr< const FUICommandInfo >& InCommand, const TAttribute<FText>& InLabelOverride = TAttribute<FText>(), const TAttribute<FText>& InToolTipOverride = TAttribute<FText>(), const TAttribute<FSlateIcon>& InIconOverride = TAttribute<FSlateIcon>(), const FName InTutorialHighlightName = NAME_None, const FName InNameOverride = NAME_None);
+	FEditorMenuEntry& AddMenuEntryWithCommandList(const TSharedPtr<const FUICommandInfo>& InCommand, const TSharedPtr<const FUICommandList>& InCommandList, const TAttribute<FText>& InLabelOverride = TAttribute<FText>(), const TAttribute<FText>& InToolTipOverride = TAttribute<FText>(), const TAttribute<FSlateIcon>& InIconOverride = TAttribute<FSlateIcon>(), const FName InTutorialHighlightName = NAME_None, const FName InNameOverride = NAME_None);
 
 	FEditorMenuEntry& AddDynamicEntry(const FName InName, const FNewEditorMenuSectionDelegate& InConstruct);
 	FEditorMenuEntry& AddDynamicEntry(const FName InName, const FNewEditorMenuDelegateLegacy& InConstruct);
diff --git a/Engine/Source/Editor/EditorMenus/Public/EditorMenuSubsystem.h b/Engine/Source/Editor/EditorMenus/Public/EditorMenuSubsystem.h
--- a/Engine/Source/Editor/EditorMenus/Public/EditorMenuSubsystem.h
+++ b/Engine/Source/Editor/EditorMenus/Public/EditorMenuSubsystem.h
@@ -313,7 +313,7 @@ private:
 
 	static void ExecuteStringCommand(const FEditorMenuStringCommand StringCommand, const FEditorMenuContext Context);
 
-	void FillMenuDynamic(FMenuBuilder& Builder, FNewEditorMenuDelegate InConstructMenu);
+	void FillMenuDynamic(FMenuBuilder& Builder, FNewEditorMenuDelegate InConstructMenu, const FEditorMenuContext Context);
 
 	void ListAllParents(const FName Name, TArray<FName>& AllParents);
 
diff --git a/Engine/Source/Editor/LevelEditor/Private/SLevelEditor.cpp b/Engine/Source/Editor/LevelEditor/Private/SLevelEditor.cpp
--- a/Engine/Source/Editor/LevelEditor/Private/SLevelEditor.cpp
+++ b/Engine/Source/Editor/LevelEditor/Private/SLevelEditor.cpp
@@ -49,7 +49,7 @@
 #include "GameFramework/WorldSettings.h"
 #include "Framework/Docking/LayoutExtender.h"
 #include "HierarchicalLODOutlinerModule.h"
-
+#include "EditorViewportCommands.h"
 
 static const FName LevelEditorBuildAndSubmitTab("LevelEditorBuildAndSubmit");
 static const FName LevelEditorStatsViewerTab("LevelEditorStatsViewer");
@@ -141,6 +141,11 @@ void SLevelEditor::BindCommands()
 		Actions.FocusAllViewportsToSelection, 
 		FExecuteAction::CreateStatic( &FLevelEditorActionCallbacks::ExecuteExecCommand, FString( TEXT("CAMERA ALIGN") ) )
 		);
+
+	LevelEditorCommands->MapAction( 
+		FEditorViewportCommands::Get().FocusViewportToSelection, 
+		FExecuteAction::CreateStatic( &FLevelEditorActionCallbacks::ExecuteExecCommand, FString( TEXT("CAMERA ALIGN ACTIVEVIEWPORTONLY") ) )
+		);
 }
 
 void SLevelEditor::RegisterMenus()
diff --git a/Engine/Source/Editor/LevelEditor/Public/LevelEditorActions.h b/Engine/Source/Editor/LevelEditor/Public/LevelEditorActions.h
--- a/Engine/Source/Editor/LevelEditor/Public/LevelEditorActions.h
+++ b/Engine/Source/Editor/LevelEditor/Public/LevelEditorActions.h
@@ -630,6 +630,7 @@ public:
 	TSharedPtr< FUICommandInfo > SelectActorsInLayers;
 
 	TSharedPtr< FUICommandInfo > FocusAllViewportsToSelection;
+	TSharedPtr< FUICommandInfo > FocusViewportToSelection;
 
         // Open merge actor command
 	TSharedPtr< FUICommandInfo > OpenMergeActor;
diff --git a/Engine/Source/Editor/MainFrame/Private/Menus/MainMenu.cpp b/Engine/Source/Editor/MainFrame/Private/Menus/MainMenu.cpp
--- a/Engine/Source/Editor/MainFrame/Private/Menus/MainMenu.cpp
+++ b/Engine/Source/Editor/MainFrame/Private/Menus/MainMenu.cpp
@@ -534,16 +534,7 @@ void FMainMenu::RegisterRecentFileAndExitMenuItems()
 
 TSharedRef< SWidget > FMainMenu::MakeMainTabMenu( const TSharedPtr<FTabManager>& TabManager, const FName MenuName, FEditorMenuContext& EditorMenuContext )
 {
-	if (GetDefault<UEditorStyleSettings>()->bShowProjectMenus)
-	{
-		IMainFrameModule& MainFrameModule = FModuleManager::GetModuleChecked<IMainFrameModule>("MainFrame");
-		EditorMenuContext.AppendCommandList(MainFrameModule.GetMainFrameCommandBindings());
-	}
-
-
-	TSharedRef< SWidget > MenuBarWidget = FMainMenu::MakeMainMenu( TabManager, MenuName, EditorMenuContext );
-
-	return MenuBarWidget;
+	return FMainMenu::MakeMainMenu( TabManager, MenuName, EditorMenuContext );
 }
 
 
