Add header forward declare
#rnx

diff --git a/Engine/Source/Editor/EditorMenus/Public/EditorMenuContext.h b/Engine/Source/Editor/EditorMenus/Public/EditorMenuContext.h
--- a/Engine/Source/Editor/EditorMenus/Public/EditorMenuContext.h
+++ b/Engine/Source/Editor/EditorMenus/Public/EditorMenuContext.h
@@ -10,6 +10,7 @@ class FUICommandList;
 class FTabManager;
 class FExtender;
 class FExtensibilityManager;
+struct FUIAction;
 
 UCLASS(BlueprintType, Abstract)
 class EDITORMENUS_API UEditorMenuContextBase : public UObject
