Forward declare FUICommandInfo
#rnx

diff --git a/Engine/Source/Editor/EditorMenus/Public/EditorMenuContext.h b/Engine/Source/Editor/EditorMenus/Public/EditorMenuContext.h
--- a/Engine/Source/Editor/EditorMenus/Public/EditorMenuContext.h
+++ b/Engine/Source/Editor/EditorMenus/Public/EditorMenuContext.h
@@ -6,11 +6,12 @@
 
 #include "EditorMenuContext.generated.h"
 
+struct FUIAction;
+class FUICommandInfo;
 class FUICommandList;
 class FTabManager;
 class FExtender;
 class FExtensibilityManager;
-struct FUIAction;
 
 UCLASS(BlueprintType, Abstract)
 class EDITORMENUS_API UEditorMenuContextBase : public UObject
