Add one additional editor menu submenu function variant
#rnx

diff --git a/Engine/Source/Editor/EditorMenus/Private/EditorMenuEntry.cpp b/Engine/Source/Editor/EditorMenus/Private/EditorMenuEntry.cpp
--- a/Engine/Source/Editor/EditorMenus/Private/EditorMenuEntry.cpp
+++ b/Engine/Source/Editor/EditorMenus/Private/EditorMenuEntry.cpp
@@ -105,6 +105,21 @@ FEditorMenuEntry FEditorMenuEntry::InitMenuEntry(const FName InName, const FEdit
 	return Entry;
 }
 
+FEditorMenuEntry FEditorMenuEntry::InitSubMenu(const FName InParentMenu, const FName InName, const TAttribute<FText>& InLabel, const TAttribute<FText>& InToolTip, const FNewEditorMenuChoice& InMakeMenu, const FEditorUIActionChoice& InAction, const EUserInterfaceActionType InUserInterfaceActionType, bool bInOpenSubMenuOnClick, const TAttribute<FSlateIcon>& InIcon, const bool bInShouldCloseWindowAfterMenuSelection)
+{
+	FEditorMenuEntry Entry(UEditorMenuSubsystem::Get()->CurrentOwner(), InName, EMultiBlockType::MenuEntry);
+	Entry.Label = InLabel;
+	Entry.ToolTip = InToolTip;
+	Entry.Icon = InIcon;	
+	Entry.Action = InAction;
+	Entry.UserInterfaceActionType = InUserInterfaceActionType;
+	Entry.bShouldCloseWindowAfterMenuSelection = bInShouldCloseWindowAfterMenuSelection;
+	Entry.SubMenuData.bIsSubMenu = true;
+	Entry.SubMenuData.ConstructMenu = InMakeMenu;
+	Entry.SubMenuData.bOpenSubMenuOnClick = bInOpenSubMenuOnClick;
+	return Entry;
+}
+
 FEditorMenuEntry FEditorMenuEntry::InitSubMenu(const FName InParentMenu, const FName InName, const TAttribute<FText>& InLabel, const TAttribute<FText>& InToolTip, const FNewEditorMenuChoice& InMakeMenu, bool bInOpenSubMenuOnClick, const TAttribute<FSlateIcon>& InIcon, const bool bInShouldCloseWindowAfterMenuSelection)
 {
 	FEditorMenuEntry Entry(UEditorMenuSubsystem::Get()->CurrentOwner(), InName, EMultiBlockType::MenuEntry);
diff --git a/Engine/Source/Editor/EditorMenus/Public/EditorMenuEntry.h b/Engine/Source/Editor/EditorMenus/Public/EditorMenuEntry.h
--- a/Engine/Source/Editor/EditorMenus/Public/EditorMenuEntry.h
+++ b/Engine/Source/Editor/EditorMenus/Public/EditorMenuEntry.h
@@ -89,6 +89,7 @@ struct EDITORMENUS_API FEditorMenuEntry
 	static FEditorMenuEntry InitMenuEntryWithCommandList(const TSharedPtr<const FUICommandInfo>& InCommand, const TSharedPtr<const FUICommandList>& InCommandList, const TAttribute<FText>& InLabelOverride = TAttribute<FText>(), const TAttribute<FText>& InToolTipOverride = TAttribute<FText>(), const TAttribute<FSlateIcon>& InIconOverride = TAttribute<FSlateIcon>(), const FName InTutorialHighlightName = NAME_None, const FName InNameOverride = NAME_None);
 	static FEditorMenuEntry InitMenuEntry(const FName InName, const FEditorUIActionChoice& InAction, const TSharedRef<SWidget>& Widget);
 
+	static FEditorMenuEntry InitSubMenu(const FName InParentMenu, const FName InName, const TAttribute<FText>& InLabel, const TAttribute<FText>& InToolTip, const FNewEditorMenuChoice& InMakeMenu, const FEditorUIActionChoice& InAction, const EUserInterfaceActionType InUserInterfaceActionType, bool bInOpenSubMenuOnClick = false, const TAttribute<FSlateIcon>& InIcon = TAttribute<FSlateIcon>(), const bool bShouldCloseWindowAfterMenuSelection = true);
 	static FEditorMenuEntry InitSubMenu(const FName InParentMenu, const FName InName, const TAttribute<FText>& InLabel, const TAttribute<FText>& InToolTip, const FNewEditorMenuChoice& InMakeMenu, bool bInOpenSubMenuOnClick = false, const TAttribute<FSlateIcon>& InIcon = TAttribute<FSlateIcon>(), const bool bShouldCloseWindowAfterMenuSelection = true);
 	static FEditorMenuEntry InitSubMenu(const FName InParentMenu, const FName InName, const FEditorUIActionChoice& InAction, const TSharedRef<SWidget>& InWidget, const FNewEditorMenuChoice& InMakeMenu, bool bShouldCloseWindowAfterMenuSelection = true);
 
